Settings: moved level browser settings into LevelBrowser module

diff --git a/Engine/Config/BaseEditorUserSettings.ini b/Engine/Config/BaseEditorUserSettings.ini
--- a/Engine/Config/BaseEditorUserSettings.ini
+++ b/Engine/Config/BaseEditorUserSettings.ini
@@ -35,14 +35,8 @@ bAllowSelectTranslucent=True
 ; True if we should move actors to their appropriate grid volume levels immediately after most operations
 bUpdateActorsInGridLevelsImmediately=False
 
-; True if the actor count is displayed in the slate level browser
-bDisplayActorCountInLevelBrowser=True
-
 ; True if the memory/size data is displayed in the slate level browser
 bDisplayMemorySizeDataInLevelBrowser=False
-
-; True if Level Paths are displayed in the slate level browser
-bDisplayPathsInLevelBrowser=False
 	
 ; When enabled, the application frame rate, memory, Unreal object count, as well as a debug console will be displayed in the main editor UI
 bShowFrameRateAndMemory=False
@@ -64,9 +58,6 @@ HoverHighlightIntensity=0.000000
 bDisplayUIExtensionPoints=False
 bUseCurvesForDistributions=False
 bAutoloadCheckedOutPackages=False
-bDisplayLightmassSizeInLevelBrowser=False
-bDisplayFileSizeInLevelBrowser=False
-bDisplayEditorOffsetInLevelBrowser=False
 
 [/Script/EditorStyle.EditorStyleSettings]
 ; Applies a color vision deficiency filter to the entire editor
@@ -87,6 +78,15 @@ bShowFriendlyNames=True
 bExpandConfigurationMenus=False
 
 
+[/Script/Levels.LevelBrowserSettings]
+; True if the actor count is displayed in the slate level browser
+bDisplayActorCount=True
+bDisplayLightmassSize=False
+bDisplayFileSize=False
+; True if Level Paths are displayed in the slate level browser
+bDisplayPaths=False
+bDisplayEditorOffset=False
+
 [/Script/SceneOutliner.SceneOutlinerSettings]
 ; True when the Scene Outliner is hiding temporary/run-time Actors
 bHideTemporaryActors=False
diff --git a/Engine/Source/Editor/Levels/Private/LevelCollectionViewModel.cpp b/Engine/Source/Editor/Levels/Private/LevelCollectionViewModel.cpp
--- a/Engine/Source/Editor/Levels/Private/LevelCollectionViewModel.cpp
+++ b/Engine/Source/Editor/Levels/Private/LevelCollectionViewModel.cpp
@@ -289,70 +289,70 @@ void FLevelCollectionViewModel::ToggleReadOnlyLevels_Executed()
 
 void FLevelCollectionViewModel::OnToggleDisplayActorCount()
 {
-	bool bDisplayActorCount = GetDisplayActorCountState();
-	GEditor->AccessEditorUserSettings().bDisplayActorCountInLevelBrowser = !bDisplayActorCount;
-	GEditor->AccessEditorUserSettings().PostEditChange();
+	ULevelBrowserSettings* Settings = GetMutableDefault<ULevelBrowserSettings>();
+	Settings->bDisplayActorCount = !Settings->bDisplayActorCount;
+	Settings->PostEditChange();
 
-	DisplayActorCountChanged.Broadcast(!bDisplayActorCount);
+	DisplayActorCountChanged.Broadcast(Settings->bDisplayActorCount);
 }
 
 bool FLevelCollectionViewModel::GetDisplayActorCountState() const
 {
-	return (GEditor->AccessEditorUserSettings().bDisplayActorCountInLevelBrowser);
+	return (GetDefault<ULevelBrowserSettings>()->bDisplayActorCount);
 }
 
 void FLevelCollectionViewModel::OnToggleLightmassSize()
 {
-	bool bDisplayLightmassSize = GetDisplayLightmassSizeState();
-	GEditor->AccessEditorUserSettings().bDisplayLightmassSizeInLevelBrowser = !bDisplayLightmassSize;
-	GEditor->AccessEditorUserSettings().PostEditChange();
+	ULevelBrowserSettings* Settings = GetMutableDefault<ULevelBrowserSettings>();
+	Settings->bDisplayLightmassSize = !GetDisplayLightmassSizeState();
+	Settings->PostEditChange();
 
-	DisplayLightmassSizeChanged.Broadcast(!bDisplayLightmassSize);
+	DisplayLightmassSizeChanged.Broadcast(Settings->bDisplayLightmassSize);
 }
 
 bool FLevelCollectionViewModel::GetDisplayLightmassSizeState() const
 {
-	return GEditor->AccessEditorUserSettings().bDisplayLightmassSizeInLevelBrowser;
+	return GetDefault<ULevelBrowserSettings>()->bDisplayLightmassSize;
 }
 
 void FLevelCollectionViewModel::OnToggleFileSize()
 {
-	bool bDisplayFileSize = GetDisplayFileSizeState();
-	GEditor->AccessEditorUserSettings().bDisplayFileSizeInLevelBrowser = !bDisplayFileSize;
-	GEditor->AccessEditorUserSettings().PostEditChange();
+	ULevelBrowserSettings* Settings = GetMutableDefault<ULevelBrowserSettings>();
+	Settings->bDisplayFileSize = !GetDisplayFileSizeState();
+	Settings->PostEditChange();
 
-	DisplayFileSizeChanged.Broadcast(!bDisplayFileSize);
+	DisplayFileSizeChanged.Broadcast(Settings->bDisplayFileSize);
 }
 
 bool FLevelCollectionViewModel::GetDisplayFileSizeState() const
 {
-	return GEditor->AccessEditorUserSettings().bDisplayFileSizeInLevelBrowser;
+	return GetDefault<ULevelBrowserSettings>()->bDisplayFileSize;
 }
 
 void FLevelCollectionViewModel::OnToggleEditorOffset()
 {
-	UEditorUserSettings& UserSettings = GEditor->AccessEditorUserSettings();
-	UserSettings.bDisplayEditorOffsetInLevelBrowser = !UserSettings.bDisplayEditorOffsetInLevelBrowser;
-	GEditor->AccessEditorUserSettings().PostEditChange();
+	ULevelBrowserSettings* Settings = GetMutableDefault<ULevelBrowserSettings>();
+	Settings->bDisplayEditorOffset = !Settings->bDisplayEditorOffset;
+	Settings->PostEditChange();
 
-	DisplayEditorOffsetChanged.Broadcast(UserSettings.bDisplayEditorOffsetInLevelBrowser);
+	DisplayEditorOffsetChanged.Broadcast(Settings->bDisplayEditorOffset);
 }
 
 bool FLevelCollectionViewModel::GetDisplayEditorOffsetState() const
 {
-	return GEditor->GetEditorUserSettings().bDisplayEditorOffsetInLevelBrowser;
+	return GetDefault<ULevelBrowserSettings>()->bDisplayEditorOffset;
 }
 
 void FLevelCollectionViewModel::OnToggleDisplayPaths()
 {
-	bool bDisplayPaths = GetDisplayPathsState();
-	GEditor->AccessEditorUserSettings().bDisplayPathsInLevelBrowser = !bDisplayPaths;
-	GEditor->AccessEditorUserSettings().PostEditChange();
+	ULevelBrowserSettings* Settings = GetMutableDefault<ULevelBrowserSettings>();
+	Settings->bDisplayPaths = !GetDisplayPathsState();
+	Settings->PostEditChange();
 }
 
 bool FLevelCollectionViewModel::GetDisplayPathsState() const
 {
-	return (GEditor->AccessEditorUserSettings().bDisplayPathsInLevelBrowser);
+	return (GetDefault<ULevelBrowserSettings>()->bDisplayPaths);
 }
 
 bool FLevelCollectionViewModel::CanShiftSelection()
diff --git a/Engine/Source/Editor/Levels/Private/LevelViewModel.cpp b/Engine/Source/Editor/Levels/Private/LevelViewModel.cpp
--- a/Engine/Source/Editor/Levels/Private/LevelViewModel.cpp
+++ b/Engine/Source/Editor/Levels/Private/LevelViewModel.cpp
@@ -74,7 +74,7 @@ FString FLevelViewModel::GetName(bool bForceDisplayPath /*=false*/, bool bDispla
 		{
 			DisplayName += Level->GetOutermost()->GetName();
 
-			bool bDisplayPathsInLevelBrowser = (GEditor->AccessEditorUserSettings().bDisplayPathsInLevelBrowser);
+			bool bDisplayPathsInLevelBrowser = (GetDefault<ULevelBrowserSettings>()->bDisplayPaths);
 
 			if ( !(bDisplayPathsInLevelBrowser || bForceDisplayPath) )
 			{
@@ -96,7 +96,7 @@ FString FLevelViewModel::GetName(bool bForceDisplayPath /*=false*/, bool bDispla
 	{
 		DisplayName += LevelStreaming->PackageName.ToString();
 
-		bool bDisplayPathsInLevelBrowser = (GEditor->AccessEditorUserSettings().bDisplayPathsInLevelBrowser);
+		bool bDisplayPathsInLevelBrowser = (GetDefault<ULevelBrowserSettings>()->bDisplayPaths);
 
 		if ( !(bDisplayPathsInLevelBrowser || bForceDisplayPath) )
 		{
@@ -129,7 +129,7 @@ FString FLevelViewModel::GetLightmassSizeString() const
 	FString MemorySizeString;
 	ULevel* CurLevel = Level.Get();
 
-	if ( CurLevel && GEditor->AccessEditorUserSettings().bDisplayLightmassSizeInLevelBrowser  )
+	if ( CurLevel && GetDefault<ULevelBrowserSettings>()->bDisplayLightmassSize  )
 	{
 		// Update metrics
 		static const float ByteConversion = 1.0f / 1024.0f;
@@ -146,7 +146,7 @@ FString FLevelViewModel::GetFileSizeString() const
 	FString MemorySizeString;
 	ULevel* CurLevel = Level.Get();
 
-	if ( CurLevel && GEditor->AccessEditorUserSettings().bDisplayFileSizeInLevelBrowser  )
+	if ( CurLevel && GetDefault<ULevelBrowserSettings>()->bDisplayFileSize  )
 	{
 		// Update metrics
 		static const float ByteConversion = 1.0f / 1024.0f;
diff --git a/Engine/Source/Editor/Levels/Private/LevelsModule.cpp b/Engine/Source/Editor/Levels/Private/LevelsModule.cpp
--- a/Engine/Source/Editor/Levels/Private/LevelsModule.cpp
+++ b/Engine/Source/Editor/Levels/Private/LevelsModule.cpp
@@ -32,3 +32,6 @@ TSharedRef< SWidget > FLevelsModule::CreateLevelBrowser()
 }
 
 
+ULevelBrowserSettings::ULevelBrowserSettings(const class FPostConstructInitializeProperties& PCIP)
+	: Super(PCIP)
+{ }
\ No newline at end of file
diff --git a/Engine/Source/Editor/Levels/Private/LevelsPrivatePCH.h b/Engine/Source/Editor/Levels/Private/LevelsPrivatePCH.h
--- a/Engine/Source/Editor/Levels/Private/LevelsPrivatePCH.h
+++ b/Engine/Source/Editor/Levels/Private/LevelsPrivatePCH.h
@@ -9,10 +9,11 @@
 #include "DelegateFilter.h"
 #include "FilterCollection.h"
 
-#include "LevelsModule.h"   
+#include "LevelsModule.h"
 
 #include "ActorsAssignedToSpecificLevelsFilter.h"      
- 
+
+#include "LevelBrowserSettings.h"
 #include "LevelViewModel.h"
 #include "LevelCollectionViewCommands.h"   
 #include "LevelCollectionViewModel.h"
diff --git a/Engine/Source/Editor/Levels/Private/SLevelsView.h b/Engine/Source/Editor/Levels/Private/SLevelsView.h
--- a/Engine/Source/Editor/Levels/Private/SLevelsView.h
+++ b/Engine/Source/Editor/Levels/Private/SLevelsView.h
@@ -155,22 +155,22 @@ public:
 		// This needs to be set to the column before the "Save" column.
 		InsertColumnIndex = 5;
 
-		if(GEditor->GetEditorUserSettings().bDisplayActorCountInLevelBrowser)
+		if(GetDefault<ULevelBrowserSettings>()->bDisplayActorCount)
 		{
 			ToggleActorCount(true);
 		}
 
-		if(GEditor->GetEditorUserSettings().bDisplayLightmassSizeInLevelBrowser)
+		if(GetDefault<ULevelBrowserSettings>()->bDisplayLightmassSize)
 		{
 			ToggleLightmassSize(true);
 		}
 
-		if(GEditor->GetEditorUserSettings().bDisplayFileSizeInLevelBrowser)
+		if(GetDefault<ULevelBrowserSettings>()->bDisplayFileSize)
 		{
 			ToggleFileSize(true);
 		}
 
-		if(GEditor->GetEditorUserSettings().bDisplayEditorOffsetInLevelBrowser)
+		if(GetDefault<ULevelBrowserSettings>()->bDisplayEditorOffset)
 		{
 			ShowEditorOffsetColumn(true);
 		}
diff --git a/Engine/Source/Editor/Levels/Private/SLevelsViewRow.h b/Engine/Source/Editor/Levels/Private/SLevelsViewRow.h
--- a/Engine/Source/Editor/Levels/Private/SLevelsViewRow.h
+++ b/Engine/Source/Editor/Levels/Private/SLevelsViewRow.h
@@ -74,21 +74,21 @@ protected:
 	/** @return the visibility for the Actor Count column */
 	EVisibility IsActorColumnVisible() const
 	{
-		return (GEditor->AccessEditorUserSettings().bDisplayActorCountInLevelBrowser)?
+		return (GetDefault<ULevelBrowserSettings>()->bDisplayActorCount)?
 			EVisibility::Visible : EVisibility::Collapsed;
 	}
 
 	/** @return the visibility for the Lightmass Size column */
 	EVisibility IsLightmassSizeColumnVisible() const
 	{
-		return (GEditor->AccessEditorUserSettings().bDisplayLightmassSizeInLevelBrowser)?
+		return (GetDefault<ULevelBrowserSettings>()->bDisplayLightmassSize)?
 			EVisibility::Visible : EVisibility::Collapsed;
 	}
 
 	/** @return the visibility for the File Size column */
 	EVisibility IsFileSizeColumnVisible() const
 	{
-		return (GEditor->AccessEditorUserSettings().bDisplayFileSizeInLevelBrowser)?
+		return (GetDefault<ULevelBrowserSettings>()->bDisplayFileSize)?
 			EVisibility::Visible : EVisibility::Collapsed;
 	}
 
diff --git a/Engine/Source/Editor/Levels/Public/LevelBrowserSettings.h b/Engine/Source/Editor/Levels/Public/LevelBrowserSettings.h
new file mode 100644
--- /dev/null
+++ b/Engine/Source/Editor/Levels/Public/LevelBrowserSettings.h
@@ -0,0 +1,42 @@
+// Copyright 1998-2014 Epic Games, Inc. All Rights Reserved.
+
+/*=============================================================================
+	LevelBrowserSettings.h: Declares the ULevelBrowserSettings class.
+=============================================================================*/
+
+#pragma once
+
+#include "LevelBrowserSettings.generated.h"
+
+
+/**
+ * Implements settings for the level browser.
+ */
+UCLASS(config=EditorUserSettings)
+class LEVELS_API ULevelBrowserSettings
+	:	public UObject
+{
+	GENERATED_UCLASS_BODY()
+
+public:
+
+	/** True if the actor count is displayed in the slate level browser */
+	UPROPERTY(config)
+	uint32 bDisplayActorCount:1;
+
+	/** True if the Lightmass Size is displayed in the slate level browser */
+	UPROPERTY(config)
+	uint32 bDisplayLightmassSize:1;
+
+	/** True if the File Size is displayed in the slate level browser */
+	UPROPERTY(config)
+	uint32 bDisplayFileSize:1;
+
+	/** True if Level Paths are displayed in the slate level browser */
+	UPROPERTY(config)
+	uint32 bDisplayPaths:1;
+
+	/** True if the Editor Offset is displayed in the slate level browser */
+	UPROPERTY(config)
+	uint32 bDisplayEditorOffset:1;
+};
diff --git a/Engine/Source/Editor/UnrealEd/Classes/Editor/EditorUserSettings.h b/Engine/Source/Editor/UnrealEd/Classes/Editor/EditorUserSettings.h
--- a/Engine/Source/Editor/UnrealEd/Classes/Editor/EditorUserSettings.h
+++ b/Engine/Source/Editor/UnrealEd/Classes/Editor/EditorUserSettings.h
@@ -75,26 +75,6 @@ class UEditorUserSettings : public UObject
 	UPROPERTY(config)
 	uint32 bAllowSelectTranslucent:1;
 
-	/** True if the actor count is displayed in the slate level browser */
-	UPROPERTY(config)
-	uint32 bDisplayActorCountInLevelBrowser:1;
-
-	/** True if the Lightmass Size is displayed in the slate level browser */
-	UPROPERTY(config)
-	uint32 bDisplayLightmassSizeInLevelBrowser:1;
-
-	/** True if the File Size is displayed in the slate level browser */
-	UPROPERTY(config)
-	uint32 bDisplayFileSizeInLevelBrowser:1;
-
-	/** True if Level Paths are displayed in the slate level browser */
-	UPROPERTY(config)
-	uint32 bDisplayPathsInLevelBrowser:1;
-
-	/** True if the Editor Offset is displayed in the slate level browser */
-	UPROPERTY(config)
-	uint32 bDisplayEditorOffsetInLevelBrowser:1;
-
 	UPROPERTY()
 	class UBlueprintPaletteFavorites* BlueprintFavorites;
 
diff --git a/Engine/Source/Editor/WorldBrowser/Private/LevelModel.cpp b/Engine/Source/Editor/WorldBrowser/Private/LevelModel.cpp
--- a/Engine/Source/Editor/WorldBrowser/Private/LevelModel.cpp
+++ b/Engine/Source/Editor/WorldBrowser/Private/LevelModel.cpp
@@ -4,6 +4,7 @@
 #include "Toolkits/AssetEditorManager.h"
 #include "LevelModel.h"
 #include "LevelCollectionModel.h"
+#include "LevelBrowserSettings.h"
 
 #define LOCTEXT_NAMESPACE "WorldBrowser"
 
@@ -686,7 +687,7 @@ FString FLevelModel::GetLightmassSizeString() const
 	FString MemorySizeString;
 	ULevel* Level = GetLevelObject();
 
-	if (Level && Editor->AccessEditorUserSettings().bDisplayLightmassSizeInLevelBrowser)
+	if (Level && GetDefault<ULevelBrowserSettings>()->bDisplayLightmassSize)
 	{
 		// Update metrics
 		static const float ByteConversion = 1.0f / 1024.0f;
@@ -703,7 +704,7 @@ FString FLevelModel::GetFileSizeString() const
 	FString MemorySizeString;
 	ULevel* Level = GetLevelObject();
 
-	if (Level && Editor->AccessEditorUserSettings().bDisplayFileSizeInLevelBrowser)
+	if (Level && GetDefault<ULevelBrowserSettings>()->bDisplayFileSize)
 	{
 		// Update metrics
 		static const float ByteConversion = 1.0f / 1024.0f;
diff --git a/Engine/Source/Editor/WorldBrowser/WorldBrowser.Build.cs b/Engine/Source/Editor/WorldBrowser/WorldBrowser.Build.cs
--- a/Engine/Source/Editor/WorldBrowser/WorldBrowser.Build.cs
+++ b/Engine/Source/Editor/WorldBrowser/WorldBrowser.Build.cs
@@ -37,6 +37,7 @@ public class WorldBrowser : ModuleRules
                 "RawMesh",
                 "MeshUtilities",
                 "LandscapeEditor",
+				"Levels",
             }
 		);
 
