UE4: "Add Existing Level" in the level browser now uses an asset picker when launching with -WorldAssets

diff --git a/Engine/Source/Editor/LevelEditor/Private/LevelEditorMenu.cpp b/Engine/Source/Editor/LevelEditor/Private/LevelEditorMenu.cpp
--- a/Engine/Source/Editor/LevelEditor/Private/LevelEditorMenu.cpp
+++ b/Engine/Source/Editor/LevelEditor/Private/LevelEditorMenu.cpp
@@ -164,6 +164,7 @@ TSharedRef< SWidget > FLevelEditorMenu::MakeLevelEditorMenu( const TSharedPtr<FU
 			{
 				FContentBrowserModule& ContentBrowserModule = FModuleManager::Get().LoadModuleChecked<FContentBrowserModule>(TEXT("ContentBrowser"));
 				FAssetPickerConfig AssetPickerConfig = FLevelEditorActionCallbacks::CreateLevelAssetPickerConfig();
+				AssetPickerConfig.OnAssetsActivated = FOnAssetsActivated::CreateStatic(&FLevelEditorActionCallbacks::OpenLevelFromAssetPicker);
 				AssetPickerConfig.ThumbnailScale = 0;
 
 				TSharedRef<SWidget> MenuContent =
diff --git a/Engine/Source/Editor/Levels/Private/LevelCollectionViewCommands.h b/Engine/Source/Editor/Levels/Private/LevelCollectionViewCommands.h
--- a/Engine/Source/Editor/Levels/Private/LevelCollectionViewCommands.h
+++ b/Engine/Source/Editor/Levels/Private/LevelCollectionViewCommands.h
@@ -43,7 +43,14 @@ PRAGMA_DISABLE_OPTIMIZATION
 		UI_COMMAND( DisplayEditorOffset, "Show Editor Offset", "If enabled, displays the editor offset for each level", EUserInterfaceActionType::ToggleButton, FInputGesture() );
 		
 		UI_COMMAND( CreateEmptyLevel, "Create Empty Level", "Creates a new empty Level", EUserInterfaceActionType::Button, FInputGesture() );
-		UI_COMMAND( AddExistingLevel, "Add Existing Level", "Adds an existing level", EUserInterfaceActionType::Button, FInputGesture() );
+		if ( FParse::Param(FCommandLine::Get(), TEXT("WorldAssets")) )
+		{
+			UI_COMMAND( AddExistingLevel, "Add Other Existing Level...", "Adds an existing level using the file explorer", EUserInterfaceActionType::Button, FInputGesture() );
+		}
+		else
+		{
+			UI_COMMAND( AddExistingLevel, "Add Existing Level", "Adds an existing level", EUserInterfaceActionType::Button, FInputGesture() );
+		}
 		UI_COMMAND( AddSelectedActorsToNewLevel, "Add Selected Actors to New Level", "Adds the actors currently selected in the active viewport to a new Level", EUserInterfaceActionType::Button, FInputGesture() );
 		UI_COMMAND( MergeSelectedLevels, "Merge Levels into New Level", "Merges the selected levels into a new level, removing the original levels from the persistent", EUserInterfaceActionType::Button, FInputGesture() );
 		UI_COMMAND( RemoveSelectedLevels, "Remove Selected Levels", "Removes selected levels from world", EUserInterfaceActionType::Button, FInputGesture() );
diff --git a/Engine/Source/Editor/Levels/Private/LevelCollectionViewModel.cpp b/Engine/Source/Editor/Levels/Private/LevelCollectionViewModel.cpp
--- a/Engine/Source/Editor/Levels/Private/LevelCollectionViewModel.cpp
+++ b/Engine/Source/Editor/Levels/Private/LevelCollectionViewModel.cpp
@@ -2073,4 +2073,40 @@ void FLevelCollectionViewModel::CacheCanExecuteSourceControlVars()
 	}
 }
 
+void FLevelCollectionViewModel::AddExistingLevelFromAssetPicker(const TArray<FAssetData>& SelectedAssets, EAssetTypeActivationMethod::Type ActivationType)
+{
+	const bool bCorrectActivationMethod = (ActivationType == EAssetTypeActivationMethod::DoubleClicked || ActivationType == EAssetTypeActivationMethod::Opened);
+	if (SelectedAssets.Num() > 0 && bCorrectActivationMethod)
+	{
+		const FAssetData& AssetData = SelectedAssets[0];
+		if (AssetData.AssetClass == UWorld::StaticClass()->GetFName())
+		{
+			// Close the menu that we were picking from
+			FSlateApplication::Get().DismissAllMenus();
+
+			// Fire ULevel::LevelDirtiedEvent when falling out of scope.
+			FScopedLevelDirtied LevelDirtyCallback;
+
+			// Add the level
+			ULevel* NewLevel = EditorLevelUtils::AddLevelToWorld(CurrentWorld.Get(), *AssetData.PackageName.ToString(), AddedLevelStreamingClass);
+			if (NewLevel)
+			{
+				LevelDirtyCallback.Request();
+			}
+
+			// For safety
+			if (GEditorModeTools().IsModeActive(FBuiltinEditorModes::EM_Landscape))
+			{
+				GEditorModeTools().ActivateMode(FBuiltinEditorModes::EM_Default);
+			}
+
+			// refresh editor windows
+			FEditorDelegates::RefreshAllBrowsers.Broadcast();
+
+			// Update volume actor visibility for each viewport since we loaded a level which could potentially contain volumes
+			GUnrealEd->UpdateVolumeActorVisibility(NULL);
+		}
+	}
+}
+
 #undef LOCTEXT_NAMESPACE
diff --git a/Engine/Source/Editor/Levels/Private/LevelCollectionViewModel.h b/Engine/Source/Editor/Levels/Private/LevelCollectionViewModel.h
--- a/Engine/Source/Editor/Levels/Private/LevelCollectionViewModel.h
+++ b/Engine/Source/Editor/Levels/Private/LevelCollectionViewModel.h
@@ -1,6 +1,8 @@
 // Copyright 1998-2014 Epic Games, Inc. All Rights Reserved.
 #pragma once
 
+#include "Developer/AssetTools/Public/IAssetTypeActions.h"
+
 typedef TFilterCollection< const TSharedPtr< FLevelViewModel >& > LevelFilterCollection;
 typedef IFilter< const TSharedPtr< FLevelViewModel >& > LevelFilter;
 
@@ -137,6 +139,9 @@ public:
 	/** Caches the variables for which SCC menu options are available */
 	void CacheCanExecuteSourceControlVars();
 
+	/** Directly adds the levels selected from the asset picker to the world */
+	void AddExistingLevelFromAssetPicker(const TArray<FAssetData>& ActivatedAssets, EAssetTypeActivationMethod::Type ActivationMethod);
+
 	/********************************************************************
 	 * EVENTS
 	 ********************************************************************/
diff --git a/Engine/Source/Editor/Levels/Private/SLevelsCommandsMenu.h b/Engine/Source/Editor/Levels/Private/SLevelsCommandsMenu.h
--- a/Engine/Source/Editor/Levels/Private/SLevelsCommandsMenu.h
+++ b/Engine/Source/Editor/Levels/Private/SLevelsCommandsMenu.h
@@ -2,6 +2,8 @@
 #pragma once
 
 #include "LevelsModule.h"
+#include "Editor/LevelEditor/Public/LevelEditorActions.h"
+#include "Editor/ContentBrowser/Public/ContentBrowserModule.h"
 
 #define LOCTEXT_NAMESPACE "LevelsCommands"
 
@@ -171,7 +173,21 @@ private:
 		MenuBuilder.BeginSection("LevelsAddLevel");
 		{
 			MenuBuilder.AddMenuEntry( Commands.CreateEmptyLevel );
-			MenuBuilder.AddMenuEntry( Commands.AddExistingLevel );
+			
+			static const bool bUsingWorldAssets = FParse::Param(FCommandLine::Get(), TEXT("WorldAssets"));
+			if (bUsingWorldAssets)
+			{
+				// Add existing level submenu.
+				MenuBuilder.AddSubMenu(
+					LOCTEXT("AddExistingLevelSubMenu", "Add Existing Level"),
+					LOCTEXT("AddExistingLevelSubMenu_ToolTip", "Select an existing sublevel"),
+					FNewMenuDelegate::CreateSP(this, &SLevelsCommandsMenu::FillAddExistingLevelSubMenu));
+			}
+			else
+			{
+				MenuBuilder.AddMenuEntry( Commands.AddExistingLevel );
+			}
+
 			MenuBuilder.AddMenuEntry( Commands.AddSelectedActorsToNewLevel );
 			MenuBuilder.AddMenuEntry( Commands.MergeSelectedLevels );
 		}
@@ -185,6 +201,36 @@ private:
 		MenuBuilder.EndSection();
 	}
 
+	/** Populates the "Add Existing Level" submenu */
+	void FillAddExistingLevelSubMenu(FMenuBuilder& MenuBuilder)
+	{
+		// Asset picker for a level
+		{
+			FContentBrowserModule& ContentBrowserModule = FModuleManager::Get().LoadModuleChecked<FContentBrowserModule>(TEXT("ContentBrowser"));
+			FAssetPickerConfig AssetPickerConfig = FLevelEditorActionCallbacks::CreateLevelAssetPickerConfig();
+			AssetPickerConfig.OnAssetsActivated = FOnAssetsActivated::CreateSP( ViewModel.Get(), &FLevelCollectionViewModel::AddExistingLevelFromAssetPicker );
+			AssetPickerConfig.ThumbnailScale = 0;
+
+			TSharedRef<SWidget> MenuContent =
+				SNew(SBox)
+				.HeightOverride(600)
+				.WidthOverride(300)
+				[
+					ContentBrowserModule.Get().CreateAssetPicker(AssetPickerConfig)
+				];
+
+			const bool bNoIndent = true;
+			MenuBuilder.AddWidget(MenuContent, FText::GetEmpty(), bNoIndent);
+		}
+
+		MenuBuilder.BeginSection("LegacyAddExistingLevel", LOCTEXT("LegacyAddExistingLevelHeading", "Other Levels"));
+		{
+			// Legacy Add Existing Level
+			MenuBuilder.AddMenuEntry(FLevelsViewCommands::Get().AddExistingLevel);
+		}
+		MenuBuilder.EndSection();
+	}
+
 	/** Populates the 'Lock' menu */
 	void FillLockMenu( class FMenuBuilder& MenuBuilder )
 	{
