UE4: When running with WorldAssets, the open level dialog has been changed back to a popup. The inline menus don't work as envisioned and it causes problems on other platforms with native menus (Mac)

diff --git a/Engine/Source/Editor/LevelEditor/Private/LevelEditor.cpp b/Engine/Source/Editor/LevelEditor/Private/LevelEditor.cpp
--- a/Engine/Source/Editor/LevelEditor/Private/LevelEditor.cpp
+++ b/Engine/Source/Editor/LevelEditor/Private/LevelEditor.cpp
@@ -496,15 +496,7 @@ void FLevelEditorModule::BindGlobalLevelEditorCommands()
 	ActionList.MapAction( Commands.BrowseAPIReference, FExecuteAction::CreateStatic( &FLevelEditorActionCallbacks::BrowseAPIReference ) );
 	ActionList.MapAction( Commands.BrowseViewportControls, FExecuteAction::CreateStatic( &FLevelEditorActionCallbacks::BrowseViewportControls ) );
 	ActionList.MapAction( Commands.NewLevel, FExecuteAction::CreateStatic( &FLevelEditorActionCallbacks::NewLevel ), FCanExecuteAction::CreateStatic( &FLevelEditorActionCallbacks::NewLevel_CanExecute ) );
-	if (FParse::Param(FCommandLine::Get(), TEXT("WorldAssets")))
-	{
-		ActionList.MapAction( Commands.OpenLevel, FExecuteAction::CreateStatic( &FLevelEditorActionCallbacks::OpenLevelPickingDialog ), FCanExecuteAction::CreateStatic( &FLevelEditorActionCallbacks::OpenLevel_CanExecute ) );
-		ActionList.MapAction( Commands.LegacyOpenLevel, FExecuteAction::CreateStatic( &FLevelEditorActionCallbacks::OpenLevel ), FCanExecuteAction::CreateStatic( &FLevelEditorActionCallbacks::OpenLevel_CanExecute ) );
-	}
-	else
-	{
-		ActionList.MapAction( Commands.OpenLevel, FExecuteAction::CreateStatic( &FLevelEditorActionCallbacks::OpenLevel ), FCanExecuteAction::CreateStatic( &FLevelEditorActionCallbacks::OpenLevel_CanExecute ) );
-	}
+	ActionList.MapAction(Commands.OpenLevel, FExecuteAction::CreateStatic(&FLevelEditorActionCallbacks::OpenLevel), FCanExecuteAction::CreateStatic(&FLevelEditorActionCallbacks::OpenLevel_CanExecute));
 	ActionList.MapAction( Commands.Save, FExecuteAction::CreateStatic( &FLevelEditorActionCallbacks::Save ) );
 	ActionList.MapAction( Commands.SaveAs, FExecuteAction::CreateStatic( &FLevelEditorActionCallbacks::SaveAs ), DefaultExecuteAction );
 	ActionList.MapAction( Commands.SaveAllLevels, FExecuteAction::CreateStatic( &FLevelEditorActionCallbacks::SaveAllLevels ) );
diff --git a/Engine/Source/Editor/LevelEditor/Private/LevelEditorActions.cpp b/Engine/Source/Editor/LevelEditor/Private/LevelEditorActions.cpp
--- a/Engine/Source/Editor/LevelEditor/Private/LevelEditorActions.cpp
+++ b/Engine/Source/Editor/LevelEditor/Private/LevelEditorActions.cpp
@@ -27,7 +27,6 @@
 #include "LightingTools.h"
 #include "MRUFavoritesList.h"
 #include "Editor/SceneOutliner/Private/SSocketChooser.h"
-#include "Runtime/AssetRegistry/Public/AssetRegistryModule.h"
 #include "SnappingUtils.h"
 #include "Layers/ILayers.h"
 #include "IPlacementModeModule.h"
@@ -183,91 +182,6 @@ bool FLevelEditorActionCallbacks::OpenLevel_CanExecute()
 	return FSlateApplication::Get().IsNormalExecution() && !GLevelEditorModeTools().IsTracking();
 }
 
-FAssetPickerConfig FLevelEditorActionCallbacks::CreateLevelAssetPickerConfig()
-{
-	FAssetPickerConfig AssetPickerConfig;
-	AssetPickerConfig.Filter.ClassNames.Add(UWorld::StaticClass()->GetFName());
-	AssetPickerConfig.bAllowDragging = false;
-	AssetPickerConfig.SelectionMode = ESelectionMode::Single;
-	AssetPickerConfig.InitialAssetViewType = EAssetViewType::List;
-	UWorld* CurrentWorld = GEditor->GetEditorWorldContext().World();
-	if (CurrentWorld)
-	{
-		AssetPickerConfig.InitialAssetSelection = FAssetData(CurrentWorld);
-	}
-
-	return AssetPickerConfig;
-}
-
-void FLevelEditorActionCallbacks::OpenLevelPickingDialog()
-{
-	const FVector2D AssetPickerSize(600.0f, 586.0f);
-
-	FMenuBuilder MenuBuilder(false, NULL);
-
-	// Create the contents of the popup
-	FContentBrowserModule& ContentBrowserModule = FModuleManager::Get().LoadModuleChecked<FContentBrowserModule>(TEXT("ContentBrowser"));
-	FAssetPickerConfig AssetPickerConfig = CreateLevelAssetPickerConfig();
-	AssetPickerConfig.OnAssetsActivated = FOnAssetsActivated::CreateStatic(&FLevelEditorActionCallbacks::OpenLevelFromAssetPicker);
-	TSharedRef<SWidget> ActualWidget = 
-		SNew(SBox)
-		.HeightOverride(AssetPickerSize.X)
-		.WidthOverride(AssetPickerSize.Y)
-		[
-			ContentBrowserModule.Get().CreateAssetPicker(AssetPickerConfig)
-		];
-
-	// Wrap the picker widget in a multibox-style menu body
-	MenuBuilder.BeginSection("AssetPickerOpenLevel", NSLOCTEXT("OpenLevelDialog", "WindowTitle", "Open Level"));
-	{
-		const bool bNoIndent = true;
-		MenuBuilder.AddWidget(ActualWidget, FText::GetEmpty(), bNoIndent);
-	}
-	MenuBuilder.EndSection();
-
-	TSharedRef<SWidget> WindowContents = MenuBuilder.MakeWidget();
-
-	// Determine where the pop-up should open
-	TSharedPtr<SWindow> ParentWindow = FSlateApplication::Get().GetActiveTopLevelWindow();
-	FVector2D WindowPosition = FSlateApplication::Get().GetCursorPos();
-	if (ParentWindow.IsValid())
-	{
-		FSlateRect ParentMonitorRect = ParentWindow->GetFullScreenInfo();
-		const FVector2D MonitorCenter((ParentMonitorRect.Right + ParentMonitorRect.Left) * 0.5f, (ParentMonitorRect.Top + ParentMonitorRect.Bottom) * 0.5f);
-		WindowPosition = MonitorCenter - AssetPickerSize * 0.5f;
-
-		// Open the pop-up
-		FPopupTransitionEffect TransitionEffect(FPopupTransitionEffect::None);
-		TSharedRef<SWindow> PopupWindow = FSlateApplication::Get().PushMenu(ParentWindow.ToSharedRef(), WindowContents, WindowPosition, TransitionEffect);
-	}
-}
-
-void FLevelEditorActionCallbacks::OpenLevelFromAssetPicker(const TArray<class FAssetData>& SelectedAssets, EAssetTypeActivationMethod::Type ActivationType)
-{
-	const bool bCorrectActivationMethod = (ActivationType == EAssetTypeActivationMethod::DoubleClicked || ActivationType == EAssetTypeActivationMethod::Opened);
-	if (SelectedAssets.Num() > 0 && bCorrectActivationMethod)
-	{
-		const FAssetData& AssetData = SelectedAssets[0];
-		if (AssetData.AssetClass == UWorld::StaticClass()->GetFName())
-		{
-			// Close the menu that we were picking from
-			FSlateApplication::Get().DismissAllMenus();
-
-			// If there are any unsaved changes to the current level, see if the user wants to save those first.
-			bool bPromptUserToSave = true;
-			bool bSaveMapPackages = true;
-			bool bSaveContentPackages = true;
-			if (FEditorFileUtils::SaveDirtyPackages(bPromptUserToSave, bSaveMapPackages, bSaveContentPackages))
-			{
-				const FString FileToOpen = FPackageName::LongPackageNameToFilename(AssetData.PackageName.ToString(), FPackageName::GetMapPackageExtension());
-				const bool bLoadAsTemplate = false;
-				const bool bShowProgress = true;
-				FEditorFileUtils::LoadMap(FileToOpen, bLoadAsTemplate, bShowProgress);
-			}
-		}
-	}
-}
-
 void FLevelEditorActionCallbacks::DeltaTransform()
 {
 	if (GUnrealEd->WarnIfLightingBuildIsCurrentlyRunning())
@@ -2620,15 +2534,7 @@ void FLevelEditorCommands::RegisterCommands()
 	UI_COMMAND( BrowseAPIReference, "API Reference...", "Opens the API reference documentation", EUserInterfaceActionType::Button, FInputGesture() );
 	UI_COMMAND( BrowseViewportControls, "Viewport Controls...", "Opens the viewport controls cheat sheet", EUserInterfaceActionType::Button, FInputGesture() );
 	UI_COMMAND( NewLevel, "New Level...", "Create a new level, or choose a level template to start from.", EUserInterfaceActionType::Button, FInputGesture( EModifierKey::Control, EKeys::N ) );
-	if (FParse::Param(FCommandLine::Get(), TEXT("WorldAssets")))
-	{
-		UI_COMMAND( OpenLevel, "Open Level...", "Loads an existing level", EUserInterfaceActionType::Button, FInputGesture( EModifierKey::Control, EKeys::O ) );
-		UI_COMMAND( LegacyOpenLevel, "Open Other Level...", "Loads an existing level using the file explorer", EUserInterfaceActionType::Button, FInputGesture() );
-	}
-	else
-	{
-		UI_COMMAND( OpenLevel, "Open Level...", "Loads an existing level", EUserInterfaceActionType::Button, FInputGesture( EModifierKey::Control, EKeys::O ) );
-	}
+	UI_COMMAND( OpenLevel, "Open Level...", "Loads an existing level", EUserInterfaceActionType::Button, FInputGesture( EModifierKey::Control, EKeys::O ) );
 	UI_COMMAND( Save, "Save", "Saves the current level to disk", EUserInterfaceActionType::Button, FInputGesture() );
 	UI_COMMAND( SaveAs, "Save As...", "Save the current level as...", EUserInterfaceActionType::Button, FInputGesture( EModifierKey::Shift|EModifierKey::Control, EKeys::S ) );
 	UI_COMMAND( SaveAllLevels, "Save All Levels", "Saves all unsaved levels to disk", EUserInterfaceActionType::Button, FInputGesture() );
diff --git a/Engine/Source/Editor/LevelEditor/Private/LevelEditorMenu.cpp b/Engine/Source/Editor/LevelEditor/Private/LevelEditorMenu.cpp
--- a/Engine/Source/Editor/LevelEditor/Private/LevelEditorMenu.cpp
+++ b/Engine/Source/Editor/LevelEditor/Private/LevelEditorMenu.cpp
@@ -130,20 +130,8 @@ TSharedRef< SWidget > FLevelEditorMenu::MakeLevelEditorMenu( const TSharedPtr<FU
 			// New Level
 			MenuBuilder.AddMenuEntry( FLevelEditorCommands::Get().NewLevel );
 
-			static const bool bUsingWorldAssets = FParse::Param(FCommandLine::Get(), TEXT("WorldAssets"));
-			if ( bUsingWorldAssets )
-			{
-				// Open level menu.
-				MenuBuilder.AddSubMenu(
-					LOCTEXT("OpenLevelSubMenu", "Open Level"),
-					LOCTEXT("OpenLevelSubMenu_ToolTip", "Select a level to load"),
-					FNewMenuDelegate::CreateStatic(&Local::MakeOpenLevelMenu), false, FSlateIcon(FEditorStyle::GetStyleSetName(), "LevelEditor.OpenLevel"));
-			}
-			else
-			{
-				// Open Level
-				MenuBuilder.AddMenuEntry( FLevelEditorCommands::Get().OpenLevel );
-			}
+			// Open Level
+			MenuBuilder.AddMenuEntry( FLevelEditorCommands::Get().OpenLevel );
 
 			// Open Asset
 			//@TODO: Doesn't work when summoned from here: MenuBuilder.AddMenuEntry( FGlobalEditorCommonCommands::Get().SummonOpenAssetDialog );
@@ -158,35 +146,6 @@ TSharedRef< SWidget > FLevelEditorMenu::MakeLevelEditorMenu( const TSharedPtr<FU
 			MenuBuilder.AddMenuEntry( FLevelEditorCommands::Get().SaveAllLevels );
 		}
 
-		static void MakeOpenLevelMenu( FMenuBuilder& MenuBuilder )
-		{
-			// Asset picker for a level
-			{
-				FContentBrowserModule& ContentBrowserModule = FModuleManager::Get().LoadModuleChecked<FContentBrowserModule>(TEXT("ContentBrowser"));
-				FAssetPickerConfig AssetPickerConfig = FLevelEditorActionCallbacks::CreateLevelAssetPickerConfig();
-				AssetPickerConfig.OnAssetsActivated = FOnAssetsActivated::CreateStatic(&FLevelEditorActionCallbacks::OpenLevelFromAssetPicker);
-				AssetPickerConfig.ThumbnailScale = 0;
-
-				TSharedRef<SWidget> MenuContent =
-					SNew(SBox)
-					.HeightOverride(600)
-					.WidthOverride(300)
-					[
-						ContentBrowserModule.Get().CreateAssetPicker(AssetPickerConfig)
-					];
-
-				const bool bNoIndent = true;
-				MenuBuilder.AddWidget(MenuContent, FText::GetEmpty(), bNoIndent);
-			}
-
-			MenuBuilder.BeginSection("LegacyOpenLevel", LOCTEXT("LegacyOpenLevelHeading", "Other Levels"));
-			{
-				// Legacy Open Level
-				MenuBuilder.AddMenuEntry( FLevelEditorCommands::Get().LegacyOpenLevel );
-			}
-			MenuBuilder.EndSection();
-		}
-
 		static void FillFavoriteLevelSubMenu( FMenuBuilder& MenuBuilder, const int32 CurFavoriteIndex )
 		{
 			IMainFrameModule& MainFrameModule = FModuleManager::LoadModuleChecked<IMainFrameModule>( "MainFrame" );
diff --git a/Engine/Source/Editor/LevelEditor/Public/LevelEditorActions.h b/Engine/Source/Editor/LevelEditor/Public/LevelEditorActions.h
--- a/Engine/Source/Editor/LevelEditor/Public/LevelEditorActions.h
+++ b/Engine/Source/Editor/LevelEditor/Public/LevelEditorActions.h
@@ -40,7 +40,6 @@ public:
 	/** Level file commands */
 	TSharedPtr< FUICommandInfo > NewLevel;
 	TSharedPtr< FUICommandInfo > OpenLevel;
-	TSharedPtr< FUICommandInfo > LegacyOpenLevel;
 	TSharedPtr< FUICommandInfo > Save;
 	TSharedPtr< FUICommandInfo > SaveAs;
 	TSharedPtr< FUICommandInfo > SaveAllLevels;
@@ -591,10 +590,6 @@ public:
 	/** Opens an existing level */
 	static void OpenLevel();
 	static bool OpenLevel_CanExecute();
-	static struct FAssetPickerConfig CreateLevelAssetPickerConfig();
-	static void OpenLevelPickingDialog();
-	static void OpenLevelFromAssetPicker(const TArray<class FAssetData>& SelectedAssets, EAssetTypeActivationMethod::Type ActivationType);
-	
 
 	/** Opens delta transform */
 	static void DeltaTransform();
diff --git a/Engine/Source/Editor/Levels/Private/LevelCollectionViewCommands.h b/Engine/Source/Editor/Levels/Private/LevelCollectionViewCommands.h
--- a/Engine/Source/Editor/Levels/Private/LevelCollectionViewCommands.h
+++ b/Engine/Source/Editor/Levels/Private/LevelCollectionViewCommands.h
@@ -43,14 +43,7 @@ PRAGMA_DISABLE_OPTIMIZATION
 		UI_COMMAND( DisplayEditorOffset, "Show Editor Offset", "If enabled, displays the editor offset for each level", EUserInterfaceActionType::ToggleButton, FInputGesture() );
 		
 		UI_COMMAND( CreateEmptyLevel, "Create Empty Level", "Creates a new empty Level", EUserInterfaceActionType::Button, FInputGesture() );
-		if ( FParse::Param(FCommandLine::Get(), TEXT("WorldAssets")) )
-		{
-			UI_COMMAND( AddExistingLevel, "Add Other Existing Level...", "Adds an existing level using the file explorer", EUserInterfaceActionType::Button, FInputGesture() );
-		}
-		else
-		{
-			UI_COMMAND( AddExistingLevel, "Add Existing Level", "Adds an existing level", EUserInterfaceActionType::Button, FInputGesture() );
-		}
+		UI_COMMAND( AddExistingLevel, "Add Existing Level", "Adds an existing level", EUserInterfaceActionType::Button, FInputGesture() );
 		UI_COMMAND( AddSelectedActorsToNewLevel, "Add Selected Actors to New Level", "Adds the actors currently selected in the active viewport to a new Level", EUserInterfaceActionType::Button, FInputGesture() );
 		UI_COMMAND( MergeSelectedLevels, "Merge Levels into New Level", "Merges the selected levels into a new level, removing the original levels from the persistent", EUserInterfaceActionType::Button, FInputGesture() );
 		UI_COMMAND( RemoveSelectedLevels, "Remove Selected Levels", "Removes selected levels from world", EUserInterfaceActionType::Button, FInputGesture() );
diff --git a/Engine/Source/Editor/Levels/Private/LevelCollectionViewModel.cpp b/Engine/Source/Editor/Levels/Private/LevelCollectionViewModel.cpp
--- a/Engine/Source/Editor/Levels/Private/LevelCollectionViewModel.cpp
+++ b/Engine/Source/Editor/Levels/Private/LevelCollectionViewModel.cpp
@@ -1368,110 +1368,108 @@ void FLevelCollectionViewModel::AddExistingLevel_Executed()
 	AddExistingLevel();
 }
 
-bool FLevelCollectionViewModel::AddExistingLevel()
+void FLevelCollectionViewModel::AddExistingLevel(bool bRemoveInvalidSelectedLevelsAfter)
 {
-	TArray<FString> OpenFilenames;
-	IDesktopPlatform* DesktopPlatform = FDesktopPlatformModule::Get();
-	bool bOpened = false;
-	if ( DesktopPlatform )
+	if ( FParse::Param(FCommandLine::Get(), TEXT("WorldAssets")) )
 	{
-		void* ParentWindowWindowHandle = NULL;
-
-		IMainFrameModule& MainFrameModule = FModuleManager::LoadModuleChecked<IMainFrameModule>(TEXT("MainFrame"));
-		const TSharedPtr<SWindow>& MainFrameParentWindow = MainFrameModule.GetParentWindow();
-		if ( MainFrameParentWindow.IsValid() && MainFrameParentWindow->GetNativeWindow().IsValid() )
-		{
-			ParentWindowWindowHandle = MainFrameParentWindow->GetNativeWindow()->GetOSWindowHandle();
-		}
-
-		bOpened = DesktopPlatform->OpenFileDialog(
-			ParentWindowWindowHandle,
-		NSLOCTEXT("UnrealEd", "Open", "Open").ToString(),
-		*FEditorDirectories::Get().GetLastDirectory(ELastDirectory::UNR),
-		TEXT(""),
-		*FEditorFileUtils::GetFilterString(FI_Load),
-			EFileDialogFlags::Multiple,
-			OpenFilenames
-			);
+		FEditorFileUtils::FOnLevelsChosen LevelsChosenDelegate = FEditorFileUtils::FOnLevelsChosen::CreateSP(this, &FLevelCollectionViewModel::HandleAddExistingLevelSelected, bRemoveInvalidSelectedLevelsAfter);
+		const bool bAllowMultipleSelection = true;
+		FEditorFileUtils::OpenLevelPickingDialog(LevelsChosenDelegate, bAllowMultipleSelection);
 	}
-
-	if( bOpened )
+	else
 	{
-		// Save the path as default for next time
-		FEditorDirectories::Get().SetLastDirectory( ELastDirectory::UNR, FPaths::GetPath( OpenFilenames[ 0 ] ) );
-
-		TArray<FString> Filenames;
-		for( int32 FileIndex = 0 ; FileIndex < OpenFilenames.Num() ; ++FileIndex )
+		TArray<FString> OpenFilenames;
+		IDesktopPlatform* DesktopPlatform = FDesktopPlatformModule::Get();
+		bool bOpened = false;
+		if ( DesktopPlatform )
 		{
-			// Strip paths from to get the level package names.
-			const FString FilePath( OpenFilenames[FileIndex] );
-
-			// make sure the level is in our package cache, because the async loading code will use this to find it
-			if (!FPaths::FileExists(FilePath))
-			{
-				FMessageDialog::Open( EAppMsgType::Ok, NSLOCTEXT("UnrealEd", "Error_LevelImportFromExternal", "Importing external sublevels is not allowed. Move the level files into the standard content directory and try again.\nAfter moving the level(s), restart the editor.") );
-				return false;				
-			}
+			void* ParentWindowWindowHandle = NULL;
 
-			FText ErrorMessage;
-			bool bFilenameIsValid = FEditorFileUtils::IsValidMapFilename(OpenFilenames[FileIndex], ErrorMessage);
-			if ( !bFilenameIsValid )
+			IMainFrameModule& MainFrameModule = FModuleManager::LoadModuleChecked<IMainFrameModule>(TEXT("MainFrame"));
+			const TSharedPtr<SWindow>& MainFrameParentWindow = MainFrameModule.GetParentWindow();
+			if ( MainFrameParentWindow.IsValid() && MainFrameParentWindow->GetNativeWindow().IsValid() )
 			{
-				// Start the loop over, prompting for save again
-				const FText DisplayFilename = FText::FromString( IFileManager::Get().ConvertToAbsolutePathForExternalAppForRead(*OpenFilenames[FileIndex]) );
-				FFormatNamedArguments Arguments;
-				Arguments.Add(TEXT("Filename"), DisplayFilename);
-				Arguments.Add(TEXT("LineTerminators"), FText::FromString(LINE_TERMINATOR LINE_TERMINATOR));
-				Arguments.Add(TEXT("ErrorMessage"), ErrorMessage);
-				const FText DisplayMessage = FText::Format( NSLOCTEXT("UnrealEd", "Error_InvalidLevelToAdd", "Unable to add streaming level {Filename}{LineTerminators}{ErrorMessage}"), Arguments );
-				FMessageDialog::Open( EAppMsgType::Ok, DisplayMessage );
-				return false;
+				ParentWindowWindowHandle = MainFrameParentWindow->GetNativeWindow()->GetOSWindowHandle();
 			}
 
-			Filenames.Add( FilePath );
+			bOpened = DesktopPlatform->OpenFileDialog(
+				ParentWindowWindowHandle,
+				NSLOCTEXT("UnrealEd", "Open", "Open").ToString(),
+				*FEditorDirectories::Get().GetLastDirectory(ELastDirectory::UNR),
+				TEXT(""),
+				*FEditorFileUtils::GetFilterString(FI_Load),
+				EFileDialogFlags::Multiple,
+				OpenFilenames
+				);
 		}
 
-		// Sort the level packages alphabetically by name.
-		Filenames.Sort();
-
-		// Fire ULevel::LevelDirtiedEvent when falling out of scope.
-		FScopedLevelDirtied LevelDirtyCallback;
-
-		// Try to add the levels that were specified in the dialog.
-		ULevel* NewLevel = nullptr;
-		for( int32 FileIndex = 0 ; FileIndex < Filenames.Num() ; ++FileIndex )
+		if( bOpened )
 		{
-			const FString& PackageName = FPackageName::FilenameToLongPackageName(Filenames[FileIndex]);
+			// Save the path as default for next time
+			FEditorDirectories::Get().SetLastDirectory( ELastDirectory::UNR, FPaths::GetPath( OpenFilenames[ 0 ] ) );
 
-			NewLevel = EditorLevelUtils::AddLevelToWorld(CurrentWorld.Get(), *PackageName, AddedLevelStreamingClass);
-			if (NewLevel)
+			TArray<FString> Filenames;
+			for( int32 FileIndex = 0 ; FileIndex < OpenFilenames.Num() ; ++FileIndex )
 			{
-				LevelDirtyCallback.Request();
+				// Strip paths from to get the level package names.
+				const FString FilePath( OpenFilenames[FileIndex] );
+
+				// make sure the level is in our package cache, because the async loading code will use this to find it
+				if (!FPaths::FileExists(FilePath))
+				{
+					FMessageDialog::Open( EAppMsgType::Ok, NSLOCTEXT("UnrealEd", "Error_LevelImportFromExternal", "Importing external sublevels is not allowed. Move the level files into the standard content directory and try again.\nAfter moving the level(s), restart the editor.") );
+					return;				
+				}
+
+				FText ErrorMessage;
+				bool bFilenameIsValid = FEditorFileUtils::IsValidMapFilename(OpenFilenames[FileIndex], ErrorMessage);
+				if ( !bFilenameIsValid )
+				{
+					// Start the loop over, prompting for save again
+					const FText DisplayFilename = FText::FromString( IFileManager::Get().ConvertToAbsolutePathForExternalAppForRead(*OpenFilenames[FileIndex]) );
+					FFormatNamedArguments Arguments;
+					Arguments.Add(TEXT("Filename"), DisplayFilename);
+					Arguments.Add(TEXT("LineTerminators"), FText::FromString(LINE_TERMINATOR LINE_TERMINATOR));
+					Arguments.Add(TEXT("ErrorMessage"), ErrorMessage);
+					const FText DisplayMessage = FText::Format( NSLOCTEXT("UnrealEd", "Error_InvalidLevelToAdd", "Unable to add streaming level {Filename}{LineTerminators}{ErrorMessage}"), Arguments );
+					FMessageDialog::Open( EAppMsgType::Ok, DisplayMessage );
+					return;
+				}
+
+				Filenames.Add( FilePath );
 			}
 
-		} // for each file
+			TArray<FString> PackageNames;
+			for ( const auto& Filename : Filenames )
+			{
+				const FString& PackageName = FPackageName::FilenameToLongPackageName(Filename);
+				PackageNames.Add(PackageName);
+			}
 
-		// Set the last loaded level to be the current level
-		if (NewLevel)
-		{
-			CurrentWorld->SetCurrentLevel(NewLevel);
-		}
+			EditorLevelUtils::AddLevelsToWorld(CurrentWorld.Get(), PackageNames, AddedLevelStreamingClass);
 
-		// For safety
-		if( GLevelEditorModeTools().IsModeActive( FBuiltinEditorModes::EM_Landscape ) )
-		{
-			GLevelEditorModeTools().ActivateDefaultMode();
+			if (bRemoveInvalidSelectedLevelsAfter)
+			{
+				RemoveInvalidSelectedLevels_Executed();
+			}
 		}
+	}
+}
 
-		// refresh editor windows
-		FEditorDelegates::RefreshAllBrowsers.Broadcast();
+void FLevelCollectionViewModel::HandleAddExistingLevelSelected(const TArray<FAssetData>& SelectedAssets, bool bRemoveInvalidSelectedLevelsAfter)
+{
+	TArray<FString> PackageNames;
+	for (const auto& AssetData : SelectedAssets)
+	{
+		PackageNames.Add(AssetData.PackageName.ToString());
 	}
 
-	// Update volume actor visibility for each viewport since we loaded a level which could potentially contain volumes
-	GUnrealEd->UpdateVolumeActorVisibility(NULL);
+	EditorLevelUtils::AddLevelsToWorld(CurrentWorld.Get(), PackageNames, AddedLevelStreamingClass);
 
-	// return whether a level was selected
-	return bOpened;
+	if (bRemoveInvalidSelectedLevelsAfter)
+	{
+		RemoveInvalidSelectedLevels_Executed();
+	}
 }
 
 void FLevelCollectionViewModel::AddSelectedActorsToNewLevel_Executed()
@@ -1482,10 +1480,8 @@ void FLevelCollectionViewModel::AddSelectedActorsToNewLevel_Executed()
 void FLevelCollectionViewModel::FixupInvalidReference_Executed()
 {
 	// Browsing is essentially the same as adding an existing level
-	if (AddExistingLevel())
-	{
-		RemoveInvalidSelectedLevels_Executed();
-	}
+	const bool bRemoveInvalidSelectedLevelsAfter = true;
+	AddExistingLevel(bRemoveInvalidSelectedLevelsAfter);
 }
 
 void FLevelCollectionViewModel::RemoveInvalidSelectedLevels_Executed()
@@ -2089,42 +2085,6 @@ void FLevelCollectionViewModel::CacheCanExecuteSourceControlVars()
 	}
 }
 
-void FLevelCollectionViewModel::AddExistingLevelFromAssetPicker(const TArray<FAssetData>& SelectedAssets, EAssetTypeActivationMethod::Type ActivationType)
-{
-	const bool bCorrectActivationMethod = (ActivationType == EAssetTypeActivationMethod::DoubleClicked || ActivationType == EAssetTypeActivationMethod::Opened);
-	if (SelectedAssets.Num() > 0 && bCorrectActivationMethod)
-	{
-		const FAssetData& AssetData = SelectedAssets[0];
-		if (AssetData.AssetClass == UWorld::StaticClass()->GetFName())
-		{
-			// Close the menu that we were picking from
-			FSlateApplication::Get().DismissAllMenus();
-
-			// Fire ULevel::LevelDirtiedEvent when falling out of scope.
-			FScopedLevelDirtied LevelDirtyCallback;
-
-			// Add the level
-			ULevel* NewLevel = EditorLevelUtils::AddLevelToWorld(CurrentWorld.Get(), *AssetData.PackageName.ToString(), AddedLevelStreamingClass);
-			if (NewLevel)
-			{
-				LevelDirtyCallback.Request();
-			}
-
-			// For safety
-			if (GLevelEditorModeTools().IsModeActive(FBuiltinEditorModes::EM_Landscape))
-			{
-				GLevelEditorModeTools().ActivateDefaultMode();
-			}
-
-			// refresh editor windows
-			FEditorDelegates::RefreshAllBrowsers.Broadcast();
-
-			// Update volume actor visibility for each viewport since we loaded a level which could potentially contain volumes
-			GUnrealEd->UpdateVolumeActorVisibility(NULL);
-		}
-	}
-}
-
 bool FLevelCollectionViewModel::IsValidMoveActorsToLevel()
 {
 	static bool bCachedIsValidActorMoveResult = false;
diff --git a/Engine/Source/Editor/Levels/Private/LevelCollectionViewModel.h b/Engine/Source/Editor/Levels/Private/LevelCollectionViewModel.h
--- a/Engine/Source/Editor/Levels/Private/LevelCollectionViewModel.h
+++ b/Engine/Source/Editor/Levels/Private/LevelCollectionViewModel.h
@@ -3,6 +3,8 @@
 
 #include "Developer/AssetTools/Public/IAssetTypeActions.h"
 
+class FAssetData;
+
 typedef TFilterCollection< const TSharedPtr< FLevelViewModel >& > LevelFilterCollection;
 typedef IFilter< const TSharedPtr< FLevelViewModel >& > LevelFilter;
 
@@ -139,9 +141,6 @@ public:
 	/** Caches the variables for which SCC menu options are available */
 	void CacheCanExecuteSourceControlVars();
 
-	/** Directly adds the levels selected from the asset picker to the world */
-	void AddExistingLevelFromAssetPicker(const TArray<FAssetData>& ActivatedAssets, EAssetTypeActivationMethod::Type ActivationMethod);
-
 	/********************************************************************
 	 * EVENTS
 	 ********************************************************************/
@@ -369,8 +368,11 @@ private:
 	/** Calls AddExistingLevel which adds an existing level; prompts for path */
 	void AddExistingLevel_Executed();
 	
-	/** Adds an existing level; prompts for path, Returns true if a level is selected */
-	bool AddExistingLevel();
+	/** Adds an existing level; prompts for path. If bRemoveInvalidSelectedLevelsAfter is true, any invalid levels are removed after */
+	void AddExistingLevel(bool bRemoveInvalidSelectedLevelsAfter = false);
+
+	/** Handler for when a level is selected after invoking AddExistingLevel */
+	void HandleAddExistingLevelSelected(const TArray<FAssetData>& SelectedAssets, bool bRemoveInvalidSelectedLevelsAfter);
 
 	/** Add Selected Actors to New Level; prompts for level save location */
 	void AddSelectedActorsToNewLevel_Executed();
diff --git a/Engine/Source/Editor/Levels/Private/SLevelsCommandsMenu.h b/Engine/Source/Editor/Levels/Private/SLevelsCommandsMenu.h
--- a/Engine/Source/Editor/Levels/Private/SLevelsCommandsMenu.h
+++ b/Engine/Source/Editor/Levels/Private/SLevelsCommandsMenu.h
@@ -173,21 +173,7 @@ private:
 		MenuBuilder.BeginSection("LevelsAddLevel");
 		{
 			MenuBuilder.AddMenuEntry( Commands.CreateEmptyLevel );
-			
-			static const bool bUsingWorldAssets = FParse::Param(FCommandLine::Get(), TEXT("WorldAssets"));
-			if (bUsingWorldAssets)
-			{
-				// Add existing level submenu.
-				MenuBuilder.AddSubMenu(
-					LOCTEXT("AddExistingLevelSubMenu", "Add Existing Level"),
-					LOCTEXT("AddExistingLevelSubMenu_ToolTip", "Select an existing sublevel"),
-					FNewMenuDelegate::CreateSP(this, &SLevelsCommandsMenu::FillAddExistingLevelSubMenu));
-			}
-			else
-			{
-				MenuBuilder.AddMenuEntry( Commands.AddExistingLevel );
-			}
-
+			MenuBuilder.AddMenuEntry( Commands.AddExistingLevel );
 			MenuBuilder.AddMenuEntry( Commands.AddSelectedActorsToNewLevel );
 			MenuBuilder.AddMenuEntry( Commands.MergeSelectedLevels );
 		}
@@ -201,36 +187,6 @@ private:
 		MenuBuilder.EndSection();
 	}
 
-	/** Populates the "Add Existing Level" submenu */
-	void FillAddExistingLevelSubMenu(FMenuBuilder& MenuBuilder)
-	{
-		// Asset picker for a level
-		{
-			FContentBrowserModule& ContentBrowserModule = FModuleManager::Get().LoadModuleChecked<FContentBrowserModule>(TEXT("ContentBrowser"));
-			FAssetPickerConfig AssetPickerConfig = FLevelEditorActionCallbacks::CreateLevelAssetPickerConfig();
-			AssetPickerConfig.OnAssetsActivated = FOnAssetsActivated::CreateSP( ViewModel.Get(), &FLevelCollectionViewModel::AddExistingLevelFromAssetPicker );
-			AssetPickerConfig.ThumbnailScale = 0;
-
-			TSharedRef<SWidget> MenuContent =
-				SNew(SBox)
-				.HeightOverride(600)
-				.WidthOverride(300)
-				[
-					ContentBrowserModule.Get().CreateAssetPicker(AssetPickerConfig)
-				];
-
-			const bool bNoIndent = true;
-			MenuBuilder.AddWidget(MenuContent, FText::GetEmpty(), bNoIndent);
-		}
-
-		MenuBuilder.BeginSection("LegacyAddExistingLevel", LOCTEXT("LegacyAddExistingLevelHeading", "Other Levels"));
-		{
-			// Legacy Add Existing Level
-			MenuBuilder.AddMenuEntry(FLevelsViewCommands::Get().AddExistingLevel);
-		}
-		MenuBuilder.EndSection();
-	}
-
 	/** Populates the 'Lock' menu */
 	void FillLockMenu( class FMenuBuilder& MenuBuilder )
 	{
diff --git a/Engine/Source/Editor/UnrealEd/Private/EditorLevelUtils.cpp b/Engine/Source/Editor/UnrealEd/Private/EditorLevelUtils.cpp
--- a/Engine/Source/Editor/UnrealEd/Private/EditorLevelUtils.cpp
+++ b/Engine/Source/Editor/UnrealEd/Private/EditorLevelUtils.cpp
@@ -121,6 +121,54 @@ namespace EditorLevelUtils
 		}
 	}
 
+	ULevel* AddLevelsToWorld(UWorld* InWorld, const TArray<FString>& LevelPackageNames, UClass* LevelStreamingClass)
+	{
+		if ( !ensure(InWorld) )
+		{
+			return nullptr;
+		}
+
+		TArray<FString> PackageNames = LevelPackageNames;
+
+		// Sort the level packages alphabetically by name.
+		PackageNames.Sort();
+
+		// Fire ULevel::LevelDirtiedEvent when falling out of scope.
+		FScopedLevelDirtied LevelDirtyCallback;
+
+		// Try to add the levels that were specified in the dialog.
+		ULevel* NewLevel = nullptr;
+		for (const auto& PackageName : PackageNames)
+		{
+			NewLevel = AddLevelToWorld(InWorld, *PackageName, LevelStreamingClass);
+			if (NewLevel)
+			{
+				LevelDirtyCallback.Request();
+			}
+
+		} // for each file
+
+		// Set the last loaded level to be the current level
+		if (NewLevel)
+		{
+			InWorld->SetCurrentLevel(NewLevel);
+		}
+
+		// For safety
+		if (GLevelEditorModeTools().IsModeActive(FBuiltinEditorModes::EM_Landscape))
+		{
+			GLevelEditorModeTools().ActivateDefaultMode();
+		}
+
+		// refresh editor windows
+		FEditorDelegates::RefreshAllBrowsers.Broadcast();
+
+		// Update volume actor visibility for each viewport since we loaded a level which could potentially contain volumes
+		GUnrealEd->UpdateVolumeActorVisibility(NULL);
+
+		return NewLevel;
+	}
+
 	ULevel* AddLevelToWorld( UWorld* InWorld, const TCHAR* LevelPackageName, UClass* LevelStreamingClass)
 	{
 		ULevel* NewLevel = NULL;
diff --git a/Engine/Source/Editor/UnrealEd/Private/FileHelpers.cpp b/Engine/Source/Editor/UnrealEd/Private/FileHelpers.cpp
--- a/Engine/Source/Editor/UnrealEd/Private/FileHelpers.cpp
+++ b/Engine/Source/Editor/UnrealEd/Private/FileHelpers.cpp
@@ -23,6 +23,8 @@
 
 #include "Dialogs/DlgPickAssetPath.h"
 
+#include "Editor/ContentBrowser/Public/ContentBrowserModule.h"
+
 DEFINE_LOG_CATEGORY_STATIC(LogFileHelpers, Log, All);
 
 //definition of flag used to do special work when we're attempting to load the "startup map"
@@ -1490,6 +1492,84 @@ bool FEditorFileUtils::PromptToCheckoutLevels(bool bCheckDirty, ULevel* Specific
 	return FEditorFileUtils::PromptToCheckoutLevels( bCheckDirty, LevelsToCheckOut );	
 }
 
+void FEditorFileUtils::OpenLevelPickingDialog(const FOnLevelsChosen& OnLevelsChosen, bool bAllowMultipleSelection)
+{
+	struct FLocal
+	{
+		static void OpenLevelFromAssetPicker(const TArray<FAssetData>& SelectedAssets, EAssetTypeActivationMethod::Type ActivationType, FOnLevelsChosen OnLevelsChosen)
+		{
+			const bool bCorrectActivationMethod = (ActivationType == EAssetTypeActivationMethod::DoubleClicked || ActivationType == EAssetTypeActivationMethod::Opened);
+			if (SelectedAssets.Num() > 0 && bCorrectActivationMethod)
+			{
+				// Close the popup and any open menus
+				FSlateApplication::Get().DismissAllMenus();
+
+				TArray<FAssetData> SelectedWorldAssets;
+				for ( const auto& AssetIt : SelectedAssets )
+				{
+					if (AssetIt.AssetClass == UWorld::StaticClass()->GetFName())
+					{
+						SelectedWorldAssets.Add(AssetIt);
+					}
+				}
+
+				OnLevelsChosen.ExecuteIfBound(SelectedWorldAssets);
+			}
+		}
+	};
+
+	const FVector2D AssetPickerSize(800.0f, 650.0f);
+
+	FMenuBuilder MenuBuilder(false, NULL);
+
+	FAssetPickerConfig AssetPickerConfig;
+	AssetPickerConfig.Filter.ClassNames.Add(UWorld::StaticClass()->GetFName());
+	AssetPickerConfig.bAllowDragging = false;
+	AssetPickerConfig.SelectionMode = bAllowMultipleSelection ? ESelectionMode::Multi : ESelectionMode::Single;
+	AssetPickerConfig.InitialAssetViewType = EAssetViewType::Tile;
+	AssetPickerConfig.ThumbnailScale = 0;
+	AssetPickerConfig.OnAssetsActivated = FOnAssetsActivated::CreateStatic(&FLocal::OpenLevelFromAssetPicker, OnLevelsChosen);
+	UWorld* CurrentWorld = GEditor->GetEditorWorldContext().World();
+	if (CurrentWorld)
+	{
+		AssetPickerConfig.InitialAssetSelection = FAssetData(CurrentWorld);
+	}
+
+	// Create the contents of the popup
+	FContentBrowserModule& ContentBrowserModule = FModuleManager::Get().LoadModuleChecked<FContentBrowserModule>(TEXT("ContentBrowser"));
+	TSharedRef<SWidget> ActualWidget =
+		SNew(SBox)
+		.HeightOverride(AssetPickerSize.X)
+		.WidthOverride(AssetPickerSize.Y)
+		[
+			ContentBrowserModule.Get().CreateAssetPicker(AssetPickerConfig)
+		];
+
+	// Wrap the picker widget in a multibox-style menu body
+	MenuBuilder.BeginSection("AssetPickerOpenLevel", NSLOCTEXT("OpenLevelDialog", "WindowTitle", "Select Level"));
+	{
+		const bool bNoIndent = true;
+		MenuBuilder.AddWidget(ActualWidget, FText::GetEmpty(), bNoIndent);
+	}
+	MenuBuilder.EndSection();
+
+	TSharedRef<SWidget> WindowContents = MenuBuilder.MakeWidget();
+
+	// Determine where the pop-up should open
+	TSharedPtr<SWindow> ParentWindow = FSlateApplication::Get().GetActiveTopLevelWindow();
+	FVector2D WindowPosition = FSlateApplication::Get().GetCursorPos();
+	if (ParentWindow.IsValid())
+	{
+		FSlateRect ParentMonitorRect = ParentWindow->GetFullScreenInfo();
+		const FVector2D MonitorCenter((ParentMonitorRect.Right + ParentMonitorRect.Left) * 0.5f, (ParentMonitorRect.Top + ParentMonitorRect.Bottom) * 0.5f);
+		WindowPosition = MonitorCenter - AssetPickerSize * 0.5f;
+
+		// Open the pop-up
+		FPopupTransitionEffect TransitionEffect(FPopupTransitionEffect::None);
+		FSlateApplication::Get().PushMenu(ParentWindow.ToSharedRef(), WindowContents, WindowPosition, TransitionEffect);
+	}
+}
+
 bool FEditorFileUtils::IsValidMapFilename(const FString& MapFilename, FText& OutErrorMessage)
 {
 	if( FPaths::GetExtension(MapFilename, true) != FPackageName::GetMapPackageExtension() )
@@ -1582,51 +1662,87 @@ void FEditorFileUtils::LoadMap()
 		return;
 	}
 
-	bool bFilenameIsValid = false;
-	FString DefaultDirectory = FEditorDirectories::Get().GetLastDirectory(ELastDirectory::LEVEL);
-
-	while( !bFilenameIsValid )
+	if (FParse::Param(FCommandLine::Get(), TEXT("WorldAssets")))
 	{
-		TArray<FString> OutFiles;
-		if ( FileDialogHelpers::OpenFiles( NSLOCTEXT("UnrealEd", "Open", "Open").ToString(), GetFilterString(FI_Load), DefaultDirectory, EFileDialogFlags::None, OutFiles) )
+		struct FLocal
 		{
-			const FString& FileToOpen = OutFiles[0];
-
-			FText ErrorMessage;
-			bFilenameIsValid = FEditorFileUtils::IsValidMapFilename(FileToOpen, ErrorMessage);
-			if ( !bFilenameIsValid )
+			static void HandleLevelsChosen(const TArray<FAssetData>& SelectedAssets)
 			{
-				// Start the loop over, prompting for load again
-				const FText DisplayFilename = FText::FromString( IFileManager::Get().ConvertToAbsolutePathForExternalAppForRead(*FileToOpen) );
-				FFormatNamedArguments Arguments;
-				Arguments.Add(TEXT("Filename"), DisplayFilename);
-				Arguments.Add(TEXT("LineTerminators"), FText::FromString(LINE_TERMINATOR LINE_TERMINATOR));
-				Arguments.Add(TEXT("ErrorMessage"), ErrorMessage);
-				const FText DisplayMessage = FText::Format( NSLOCTEXT("LoadMap", "InvalidMapName", "Failed to load map {Filename}{LineTerminators}{ErrorMessage}"), Arguments );
-				FMessageDialog::Open( EAppMsgType::Ok, DisplayMessage );
-				continue;
+				if ( SelectedAssets.Num() > 0 )
+				{
+					const FAssetData& AssetData = SelectedAssets[0];
+
+					if (!GIsDemoMode)
+					{
+						// If there are any unsaved changes to the current level, see if the user wants to save those first.
+						bool bPromptUserToSave = true;
+						bool bSaveMapPackages = true;
+						bool bSaveContentPackages = true;
+						if (FEditorFileUtils::SaveDirtyPackages(bPromptUserToSave, bSaveMapPackages, bSaveContentPackages) == false)
+						{
+							return;
+						}
+					}
+
+					const FString FileToOpen = FPackageName::LongPackageNameToFilename(AssetData.PackageName.ToString(), FPackageName::GetMapPackageExtension());
+					const bool bLoadAsTemplate = false;
+					const bool bShowProgress = true;
+					FEditorFileUtils::LoadMap(FileToOpen, bLoadAsTemplate, bShowProgress);
+				}
 			}
+		};
+		
+		const bool bAllowMultipleSelection = false;
+		OpenLevelPickingDialog(FOnLevelsChosen::CreateStatic(&FLocal::HandleLevelsChosen), bAllowMultipleSelection);
+	}
+	else
+	{
+		bool bFilenameIsValid = false;
+		FString DefaultDirectory = FEditorDirectories::Get().GetLastDirectory(ELastDirectory::LEVEL);
 
-			if( !GIsDemoMode )
+		while( !bFilenameIsValid )
+		{
+			TArray<FString> OutFiles;
+			if ( FileDialogHelpers::OpenFiles( NSLOCTEXT("UnrealEd", "Open", "Open").ToString(), GetFilterString(FI_Load), DefaultDirectory, EFileDialogFlags::None, OutFiles) )
 			{
-				// If there are any unsaved changes to the current level, see if the user wants to save those first.
-				bool bPromptUserToSave = true;
-				bool bSaveMapPackages = true;
-				bool bSaveContentPackages = true;
-				if( FEditorFileUtils::SaveDirtyPackages(bPromptUserToSave, bSaveMapPackages, bSaveContentPackages) == false )
+				const FString& FileToOpen = OutFiles[0];
+
+				FText ErrorMessage;
+				bFilenameIsValid = FEditorFileUtils::IsValidMapFilename(FileToOpen, ErrorMessage);
+				if ( !bFilenameIsValid )
 				{
-					// something went wrong or the user pressed cancel.  Return to the editor so the user doesn't lose their changes		
-					return;
+					// Start the loop over, prompting for load again
+					const FText DisplayFilename = FText::FromString( IFileManager::Get().ConvertToAbsolutePathForExternalAppForRead(*FileToOpen) );
+					FFormatNamedArguments Arguments;
+					Arguments.Add(TEXT("Filename"), DisplayFilename);
+					Arguments.Add(TEXT("LineTerminators"), FText::FromString(LINE_TERMINATOR LINE_TERMINATOR));
+					Arguments.Add(TEXT("ErrorMessage"), ErrorMessage);
+					const FText DisplayMessage = FText::Format( NSLOCTEXT("LoadMap", "InvalidMapName", "Failed to load map {Filename}{LineTerminators}{ErrorMessage}"), Arguments );
+					FMessageDialog::Open( EAppMsgType::Ok, DisplayMessage );
+					continue;
 				}
-			}
 
-			FEditorDirectories::Get().SetLastDirectory(ELastDirectory::LEVEL, FPaths::GetPath(FileToOpen));
-			LoadMap( FileToOpen, false, true );
-		}
-		else
-		{
-			// User canceled the open dialog, do not prompt again.
-			break;
+				if( !GIsDemoMode )
+				{
+					// If there are any unsaved changes to the current level, see if the user wants to save those first.
+					bool bPromptUserToSave = true;
+					bool bSaveMapPackages = true;
+					bool bSaveContentPackages = true;
+					if( FEditorFileUtils::SaveDirtyPackages(bPromptUserToSave, bSaveMapPackages, bSaveContentPackages) == false )
+					{
+						// something went wrong or the user pressed cancel.  Return to the editor so the user doesn't lose their changes		
+						return;
+					}
+				}
+
+				FEditorDirectories::Get().SetLastDirectory(ELastDirectory::LEVEL, FPaths::GetPath(FileToOpen));
+				LoadMap( FileToOpen, false, true );
+			}
+			else
+			{
+				// User canceled the open dialog, do not prompt again.
+				break;
+			}
 		}
 	}
 }
diff --git a/Engine/Source/Editor/UnrealEd/Public/EditorLevelUtils.h b/Engine/Source/Editor/UnrealEd/Public/EditorLevelUtils.h
--- a/Engine/Source/Editor/UnrealEd/Public/EditorLevelUtils.h
+++ b/Engine/Source/Editor/UnrealEd/Public/EditorLevelUtils.h
@@ -24,6 +24,17 @@ namespace EditorLevelUtils
 	 */
 	void MovesActorsToLevel( TArray< AActor* >& ActorsToMove, ULevelStreaming* DestLevelStreaming, int32& OutNumMovedActors );
 
+	/**
+	 * Adds the named level packages to the world.  Does nothing if all the levels already exist in the world.
+	 *
+	 * @param	InWorld				World in which to add the level
+	 * @param	LevelPackageName	The base filename of the level package to add.
+	 * @param	LevelStreamingClass	The streaming class type instead to use for the level.
+	 *
+	 * @return								The new level, or NULL if the level couldn't added.
+	 */
+	UNREALED_API ULevel* AddLevelsToWorld(UWorld* InWorld, const TArray<FString>& LevelPackageNames, UClass* LevelStreamingClass);
+
 
 	/**
 	 * Adds the named level package to the world.  Does nothing if the level already exists in the world.
diff --git a/Engine/Source/Editor/UnrealEd/Public/FileHelpers.h b/Engine/Source/Editor/UnrealEd/Public/FileHelpers.h
--- a/Engine/Source/Editor/UnrealEd/Public/FileHelpers.h
+++ b/Engine/Source/Editor/UnrealEd/Public/FileHelpers.h
@@ -7,6 +7,7 @@
 // Forward declarations.
 class FFileName;
 class FString;
+class FAssetData;
 
 enum EFileInteraction
 {
@@ -46,6 +47,15 @@ public:
 	////////////////////////////////////////////////////////////////////////////
 	// Loading
 
+	DECLARE_DELEGATE_OneParam(FOnLevelsChosen, const TArray<FAssetData>& /*SelectedLevels*/);
+	/**
+	 * Opens a non-modal dialog to allow the user to choose a level
+	 *
+	 * @param	OnLevelsChosen				Delegate executed when one more more levels have been selected
+	 * @param	bAllowMultipleSelection		If true, more than one level can be chosen
+	 */
+	static UNREALED_API void OpenLevelPickingDialog(const FOnLevelsChosen& OnLevelsChosen, bool bAllowMultipleSelection);
+
 	/**
 	 * Returns true if the specified map filename is valid for loading or saving.
 	 * When returning false, OutErrorMessage is supplied with a display string describing the reason why the map name is invalid.
diff --git a/Engine/Source/Editor/WorldBrowser/Private/StreamingLevels/StreamingLevelCollectionModel.cpp b/Engine/Source/Editor/WorldBrowser/Private/StreamingLevels/StreamingLevelCollectionModel.cpp
--- a/Engine/Source/Editor/WorldBrowser/Private/StreamingLevels/StreamingLevelCollectionModel.cpp
+++ b/Engine/Source/Editor/WorldBrowser/Private/StreamingLevels/StreamingLevelCollectionModel.cpp
@@ -416,113 +416,116 @@ void FStreamingLevelCollectionModel::AddExistingLevel_Executed()
 	AddExistingLevel();
 }
 
-bool FStreamingLevelCollectionModel::AddExistingLevel()
+void FStreamingLevelCollectionModel::AddExistingLevel(bool bRemoveInvalidSelectedLevelsAfter)
 {
-	TArray<FString> OpenFilenames;
-	IDesktopPlatform* DesktopPlatform = FDesktopPlatformModule::Get();
-	bool bOpened = false;
-	if ( DesktopPlatform )
+	if (FParse::Param(FCommandLine::Get(), TEXT("WorldAssets")))
 	{
-		void* ParentWindowWindowHandle = NULL;
-
-		IMainFrameModule& MainFrameModule = FModuleManager::LoadModuleChecked<IMainFrameModule>(TEXT("MainFrame"));
-		const TSharedPtr<SWindow>& MainFrameParentWindow = MainFrameModule.GetParentWindow();
-		if ( MainFrameParentWindow.IsValid() && MainFrameParentWindow->GetNativeWindow().IsValid() )
-		{
-			ParentWindowWindowHandle = MainFrameParentWindow->GetNativeWindow()->GetOSWindowHandle();
-		}
-
-		bOpened = DesktopPlatform->OpenFileDialog(
-			ParentWindowWindowHandle,
-		NSLOCTEXT("UnrealEd", "Open", "Open").ToString(),
-		*FEditorDirectories::Get().GetLastDirectory(ELastDirectory::UNR),
-		TEXT(""),
-		*FEditorFileUtils::GetFilterString(FI_Load),
-			EFileDialogFlags::Multiple,
-			OpenFilenames
-			);
+		FEditorFileUtils::FOnLevelsChosen LevelsChosenDelegate = FEditorFileUtils::FOnLevelsChosen::CreateSP(this, &FStreamingLevelCollectionModel::HandleAddExistingLevelSelected, bRemoveInvalidSelectedLevelsAfter);
+		const bool bAllowMultipleSelection = true;
+		FEditorFileUtils::OpenLevelPickingDialog(LevelsChosenDelegate, bAllowMultipleSelection);
 	}
-
-	if( bOpened )
+	else
 	{
-		// Save the path as default for next time
-		FEditorDirectories::Get().SetLastDirectory( ELastDirectory::UNR, FPaths::GetPath( OpenFilenames[ 0 ] ) );
-
-		TArray<FString> Filenames;
-		for( int32 FileIndex = 0 ; FileIndex < OpenFilenames.Num() ; ++FileIndex )
+		TArray<FString> OpenFilenames;
+		IDesktopPlatform* DesktopPlatform = FDesktopPlatformModule::Get();
+		bool bOpened = false;
+		if ( DesktopPlatform )
 		{
-			// Strip paths from to get the level package names.
-			const FString FilePath( OpenFilenames[FileIndex] );
+			void* ParentWindowWindowHandle = NULL;
 
-			// make sure the level is in our package cache, because the async loading code will use this to find it
-			if (!FPaths::FileExists(FilePath))
+			IMainFrameModule& MainFrameModule = FModuleManager::LoadModuleChecked<IMainFrameModule>(TEXT("MainFrame"));
+			const TSharedPtr<SWindow>& MainFrameParentWindow = MainFrameModule.GetParentWindow();
+			if ( MainFrameParentWindow.IsValid() && MainFrameParentWindow->GetNativeWindow().IsValid() )
 			{
-				FMessageDialog::Open( EAppMsgType::Ok, NSLOCTEXT("UnrealEd", "Error_LevelImportFromExternal", "Importing external sublevels is not allowed. Move the level files into the standard content directory and try again.\nAfter moving the level(s), restart the editor.") );
-				return false;				
+				ParentWindowWindowHandle = MainFrameParentWindow->GetNativeWindow()->GetOSWindowHandle();
 			}
 
-			FText ErrorMessage;
-			bool bFilenameIsValid = FEditorFileUtils::IsValidMapFilename(OpenFilenames[FileIndex], ErrorMessage);
-			if ( !bFilenameIsValid )
-			{
-				// Start the loop over, prompting for save again
-				const FText DisplayFilename = FText::FromString( IFileManager::Get().ConvertToAbsolutePathForExternalAppForRead(*OpenFilenames[FileIndex]) );
-				FFormatNamedArguments Arguments;
-				Arguments.Add(TEXT("Filename"), DisplayFilename);
-				Arguments.Add(TEXT("LineTerminators"), FText::FromString(LINE_TERMINATOR LINE_TERMINATOR));
-				Arguments.Add(TEXT("ErrorMessage"), ErrorMessage);
-				const FText DisplayMessage = FText::Format( NSLOCTEXT("UnrealEd", "Error_InvalidLevelToAdd", "Unable to add streaming level {Filename}{LineTerminators}{ErrorMessage}"), Arguments );
-				FMessageDialog::Open( EAppMsgType::Ok, DisplayMessage );
-				return false;
-			}
-
-			Filenames.Add( FilePath );
+			bOpened = DesktopPlatform->OpenFileDialog(
+				ParentWindowWindowHandle,
+				NSLOCTEXT("UnrealEd", "Open", "Open").ToString(),
+				*FEditorDirectories::Get().GetLastDirectory(ELastDirectory::UNR),
+				TEXT(""),
+				*FEditorFileUtils::GetFilterString(FI_Load),
+				EFileDialogFlags::Multiple,
+				OpenFilenames
+				);
 		}
 
-		// Sort the level packages alphabetically by name.
-		Filenames.Sort();
-
-		// Fire ULevel::LevelDirtiedEvent when falling out of scope.
-		FScopedLevelDirtied LevelDirtyCallback;
-
-		// Try to add the levels that were specified in the dialog.
-		ULevel* NewLevel = nullptr;
-		for( int32 FileIndex = 0 ; FileIndex < Filenames.Num() ; ++FileIndex )
+		if( bOpened )
 		{
-			const FString& PackageName = FPackageName::FilenameToLongPackageName(Filenames[FileIndex]);
+			// Save the path as default for next time
+			FEditorDirectories::Get().SetLastDirectory( ELastDirectory::UNR, FPaths::GetPath( OpenFilenames[ 0 ] ) );
 
-			NewLevel = EditorLevelUtils::AddLevelToWorld(CurrentWorld.Get(), *PackageName, AddedLevelStreamingClass);
-			if (NewLevel)
+			TArray<FString> Filenames;
+			for( int32 FileIndex = 0 ; FileIndex < OpenFilenames.Num() ; ++FileIndex )
 			{
-				LevelDirtyCallback.Request();
+				// Strip paths from to get the level package names.
+				const FString FilePath( OpenFilenames[FileIndex] );
+
+				// make sure the level is in our package cache, because the async loading code will use this to find it
+				if (!FPaths::FileExists(FilePath))
+				{
+					FMessageDialog::Open( EAppMsgType::Ok, NSLOCTEXT("UnrealEd", "Error_LevelImportFromExternal", "Importing external sublevels is not allowed. Move the level files into the standard content directory and try again.\nAfter moving the level(s), restart the editor.") );
+					return;				
+				}
+
+				FText ErrorMessage;
+				bool bFilenameIsValid = FEditorFileUtils::IsValidMapFilename(OpenFilenames[FileIndex], ErrorMessage);
+				if ( !bFilenameIsValid )
+				{
+					// Start the loop over, prompting for save again
+					const FText DisplayFilename = FText::FromString( IFileManager::Get().ConvertToAbsolutePathForExternalAppForRead(*OpenFilenames[FileIndex]) );
+					FFormatNamedArguments Arguments;
+					Arguments.Add(TEXT("Filename"), DisplayFilename);
+					Arguments.Add(TEXT("LineTerminators"), FText::FromString(LINE_TERMINATOR LINE_TERMINATOR));
+					Arguments.Add(TEXT("ErrorMessage"), ErrorMessage);
+					const FText DisplayMessage = FText::Format( NSLOCTEXT("UnrealEd", "Error_InvalidLevelToAdd", "Unable to add streaming level {Filename}{LineTerminators}{ErrorMessage}"), Arguments );
+					FMessageDialog::Open( EAppMsgType::Ok, DisplayMessage );
+					return;
+				}
+
+				Filenames.Add( FilePath );
 			}
 
-		} // for each file
+			TArray<FString> PackageNames;
+			for (const auto& Filename : Filenames)
+			{
+				const FString& PackageName = FPackageName::FilenameToLongPackageName(Filename);
+				PackageNames.Add(PackageName);
+			}
 
-		// Force a cached level list rebuild
-		PopulateLevelsList();
+			// Save or selected list, adding a new level will clean it up
+			FLevelModelList SavedInvalidSelectedLevels = InvalidSelectedLevels;
 
-		// Set the last loaded level to be the current level
-		if (NewLevel)
-		{
-			CurrentWorld->SetCurrentLevel(NewLevel);
-		}
+			EditorLevelUtils::AddLevelsToWorld(CurrentWorld.Get(), PackageNames, AddedLevelStreamingClass);
 
-		// For safety
-		if( GLevelEditorModeTools().IsModeActive( FBuiltinEditorModes::EM_Landscape ) )
-		{
-			GLevelEditorModeTools().ActivateDefaultMode();
+			if (bRemoveInvalidSelectedLevelsAfter)
+			{
+				InvalidSelectedLevels = SavedInvalidSelectedLevels;
+				RemoveInvalidSelectedLevels_Executed();
+			}
 		}
+	}
+}
 
-		// refresh editor windows
-		FEditorDelegates::RefreshAllBrowsers.Broadcast();
+void FStreamingLevelCollectionModel::HandleAddExistingLevelSelected(const TArray<FAssetData>& SelectedAssets, bool bRemoveInvalidSelectedLevelsAfter)
+{
+	TArray<FString> PackageNames;
+	for (const auto& AssetData : SelectedAssets)
+	{
+		PackageNames.Add(AssetData.PackageName.ToString());
 	}
 
-	// Update volume actor visibility for each viewport since we loaded a level which could potentially contain volumes
-	GUnrealEd->UpdateVolumeActorVisibility(NULL);
+	// Save or selected list, adding a new level will clean it up
+	FLevelModelList SavedInvalidSelectedLevels = InvalidSelectedLevels;
+
+	EditorLevelUtils::AddLevelsToWorld(CurrentWorld.Get(), PackageNames, AddedLevelStreamingClass);
 
-	// return whether a level was selected
-	return bOpened;
+	if (bRemoveInvalidSelectedLevelsAfter)
+	{
+		InvalidSelectedLevels = SavedInvalidSelectedLevels;
+		RemoveInvalidSelectedLevels_Executed();
+	}
 }
 
 void FStreamingLevelCollectionModel::AddSelectedActorsToNewLevel_Executed()
@@ -535,15 +538,9 @@ void FStreamingLevelCollectionModel::AddSelectedActorsToNewLevel_Executed()
 
 void FStreamingLevelCollectionModel::FixupInvalidReference_Executed()
 {
-	// Save or selected list, adding a new level will clean it up
-	FLevelModelList SavedInvalidSelectedLevels = InvalidSelectedLevels;
-
 	// Browsing is essentially the same as adding an existing level
-	if (AddExistingLevel())
-	{
-		InvalidSelectedLevels = SavedInvalidSelectedLevels;
-		RemoveInvalidSelectedLevels_Executed();
-	}
+	const bool bRemoveInvalidSelectedLevelsAfter = true;
+	AddExistingLevel(bRemoveInvalidSelectedLevelsAfter);
 }
 
 void FStreamingLevelCollectionModel::RemoveInvalidSelectedLevels_Executed()
diff --git a/Engine/Source/Editor/WorldBrowser/Private/StreamingLevels/StreamingLevelCollectionModel.h b/Engine/Source/Editor/WorldBrowser/Private/StreamingLevels/StreamingLevelCollectionModel.h
--- a/Engine/Source/Editor/WorldBrowser/Private/StreamingLevels/StreamingLevelCollectionModel.h
+++ b/Engine/Source/Editor/WorldBrowser/Private/StreamingLevels/StreamingLevelCollectionModel.h
@@ -86,7 +86,10 @@ private:
 	void AddExistingLevel_Executed();
 	
 	/** Adds an existing level; prompts for path, Returns true if a level is selected */
-	bool AddExistingLevel();
+	void AddExistingLevel(bool bRemoveInvalidSelectedLevelsAfter = false);
+
+	/** Handler for when a level is selected after invoking AddExistingLevel */
+	void HandleAddExistingLevelSelected(const TArray<FAssetData>& SelectedAssets, bool bRemoveInvalidSelectedLevelsAfter);
 
 	/** Add Selected Actors to New Level; prompts for level save location */
 	void AddSelectedActorsToNewLevel_Executed();
