Disallow moving actors to another level if any actor is referenced by an open Matinee

#ttp 335669 - EDITOR: MATINEE: With matinee open if you move a matinee actor to the same or new level, it duplicates it

#branch UE4

#change Check to see if any of the selected actors are referenced in an open matinee, and ask the user if they want to close matinee to continue

diff --git a/Engine/Source/Editor/Levels/Private/LevelCollectionViewModel.cpp b/Engine/Source/Editor/Levels/Private/LevelCollectionViewModel.cpp
--- a/Engine/Source/Editor/Levels/Private/LevelCollectionViewModel.cpp
+++ b/Engine/Source/Editor/Levels/Private/LevelCollectionViewModel.cpp
@@ -1098,8 +1098,38 @@ void FLevelCollectionViewModel::EditProperties_Executed()
 void FLevelCollectionViewModel::MoveActorsToSelected_Executed()
 {
 	// We shouldn't be able to get here if we have more than 1 level selected. But just in case lets check.
-	if( SelectedLevels.Num() == 1 ) 
-	{		
+	if (SelectedLevels.Num() == 1)
+	{
+		// If matinee is open, and if an actor being moved belongs to it, message the user
+		if (GLevelEditorModeTools().IsModeActive(FBuiltinEditorModes::EM_InterpEdit))
+		{
+			const FEdModeInterpEdit* InterpEditMode = (const FEdModeInterpEdit*)GLevelEditorModeTools().GetActiveMode(FBuiltinEditorModes::EM_InterpEdit);
+			if (InterpEditMode && InterpEditMode->MatineeActor)
+			{
+				TArray<AActor*> ControlledActors;
+				InterpEditMode->MatineeActor->GetControlledActors(ControlledActors);
+				if (ControlledActors.Num() > 0)
+				{
+					// are any of the selected actors in the matinee
+					USelection* SelectedActors = GEditor->GetSelectedActors();
+					for (FSelectionIterator Iter(*SelectedActors); Iter; ++Iter)
+					{
+						AActor* Actor = CastChecked<AActor>(*Iter);
+						if (Actor != nullptr && ControlledActors.Contains(Actor))
+						{
+							const bool ExitInterp = EAppReturnType::Yes == FMessageDialog::Open(EAppMsgType::YesNo, NSLOCTEXT("UnrealEd", "MatineeUnableToMove", "You must close Matinee before moving actors.\nDo you wish to do this now and continue?"));
+							if (!ExitInterp)
+							{
+								return;
+							}
+							GLevelEditorModeTools().DeactivateMode(FBuiltinEditorModes::EM_InterpEdit);
+							break;
+						}
+					}
+				}
+			}
+		}
+		
 		GEditor->MoveSelectedActorsToLevel( SelectedLevels[0]->GetLevel().Get() );
 	}	
 }
diff --git a/Engine/Source/Editor/WorldBrowser/Private/LevelCollectionModel.cpp b/Engine/Source/Editor/WorldBrowser/Private/LevelCollectionModel.cpp
--- a/Engine/Source/Editor/WorldBrowser/Private/LevelCollectionModel.cpp
+++ b/Engine/Source/Editor/WorldBrowser/Private/LevelCollectionModel.cpp
@@ -1432,6 +1432,36 @@ void FLevelCollectionModel::MakeLevelCurrent_Executed()
 
 void FLevelCollectionModel::MoveActorsToSelected_Executed()
 {
+	// If matinee is open, and if an actor being moved belongs to it, message the user
+	if (GLevelEditorModeTools().IsModeActive(FBuiltinEditorModes::EM_InterpEdit))
+	{
+		const FEdModeInterpEdit* InterpEditMode = (const FEdModeInterpEdit*)GLevelEditorModeTools().GetActiveMode(FBuiltinEditorModes::EM_InterpEdit);
+		if (InterpEditMode && InterpEditMode->MatineeActor)
+		{
+			TArray<AActor*> ControlledActors;
+			InterpEditMode->MatineeActor->GetControlledActors(ControlledActors);
+			if (ControlledActors.Num() > 0)
+			{
+				// are any of the selected actors in the matinee
+				USelection* SelectedActors = GEditor->GetSelectedActors();
+				for (FSelectionIterator Iter(*SelectedActors); Iter; ++Iter)
+				{
+					AActor* Actor = CastChecked<AActor>(*Iter);
+					if (Actor != nullptr && ControlledActors.Contains(Actor))
+					{
+						const bool ExitInterp = EAppReturnType::Yes == FMessageDialog::Open(EAppMsgType::YesNo, NSLOCTEXT("UnrealEd", "MatineeUnableToMove", "You must close Matinee before moving actors.\nDo you wish to do this now and continue?"));
+						if (!ExitInterp)
+						{
+							return;
+						}
+						GLevelEditorModeTools().DeactivateMode(FBuiltinEditorModes::EM_InterpEdit);
+						break;
+					}
+				}
+			}
+		}
+	}
+
 	MakeLevelCurrent_Executed();
 	const FScopedTransaction Transaction( LOCTEXT("MoveSelectedActorsToSelectedLevel", "Move Selected Actors to Level") );
 	Editor->MoveSelectedActorsToLevel(GetWorld()->GetCurrentLevel());
