Removed "InterpreteAsLinear", Deprecated FColor DrawColor and added FLinearColor LevelColor to replace it. Changed the underlying FColor's to FLinearColor in the rendering classes, cleaned up a bunch of color conversions.

diff --git a/Engine/Source/Editor/DetailCustomizations/Private/MathStructCustomizations.cpp b/Engine/Source/Editor/DetailCustomizations/Private/MathStructCustomizations.cpp
--- a/Engine/Source/Editor/DetailCustomizations/Private/MathStructCustomizations.cpp
+++ b/Engine/Source/Editor/DetailCustomizations/Private/MathStructCustomizations.cpp
@@ -356,7 +356,6 @@ void FColorStructCustomization::MakeHeaderRow( TSharedRef<class IPropertyHandle>
 	// We'll set up reset to default ourselves
 	const bool bDisplayResetToDefault = false;
 	const FString DisplayNameOverride = TEXT("");
-	const bool bInterpretAsLinear = StructPropertyHandle->HasMetaData("InterpretAsLinear");
 
 	FSlateFontInfo NormalText = IDetailLayoutBuilder::GetDetailFont();
 
@@ -380,7 +379,6 @@ void FColorStructCustomization::MakeHeaderRow( TSharedRef<class IPropertyHandle>
 				.Color( this, &FColorStructCustomization::OnGetColorForColorBlock )
 				.ShowBackgroundForAlpha(true)
 				.IgnoreAlpha( bIgnoreAlpha )
-				.UseSRGB(!bInterpretAsLinear)
 				.OnMouseButtonDown( this, &FColorStructCustomization::OnMouseButtonDownColorBlock )
 				.Size( FVector2D( 35.0f, 12.0f ) )
 			]
@@ -403,7 +401,6 @@ void FColorStructCustomization::MakeHeaderRow( TSharedRef<class IPropertyHandle>
 			.Color( this, &FColorStructCustomization::OnGetColorForColorBlock )
 			.ShowBackgroundForAlpha(false)
 			.IgnoreAlpha(true)
-			.UseSRGB(!bInterpretAsLinear)
 			.OnMouseButtonDown( this, &FColorStructCustomization::OnMouseButtonDownColorBlock )
 			.Size( FVector2D( 35.0f, 12.0f ) )
 		]
diff --git a/Engine/Source/Editor/Levels/Private/LevelViewModel.cpp b/Engine/Source/Editor/Levels/Private/LevelViewModel.cpp
--- a/Engine/Source/Editor/Levels/Private/LevelViewModel.cpp
+++ b/Engine/Source/Editor/Levels/Private/LevelViewModel.cpp
@@ -451,11 +451,11 @@ FSlateColor FLevelViewModel::GetColor() const
 		ULevelStreaming* StreamingLevel = FLevelUtils::FindStreamingLevel( Level.Get() );
 		if ( StreamingLevel )
 		{
-			return StreamingLevel->DrawColor.ReinterpretAsLinear();
+			return StreamingLevel->LevelColor;
 		}
 	}
 
-	return FLinearColor( 1.0f, 1.0f, 1.0f, 1.0f );
+	return FLinearColor::White;
 }
 
 void FLevelViewModel::OnColorPickerCancelled(FLinearColor Color)
@@ -476,14 +476,14 @@ void FLevelViewModel::ChangeColor(const TSharedRef<SWidget>& InPickerParentWidge
 		ULevelStreaming* StreamingLevel = FLevelUtils::FindStreamingLevel( Level.Get() );
 		check( StreamingLevel );
 
-		FColor NewColor = StreamingLevel->DrawColor;
-		TArray<FColor*> ColorArray;
+		FLinearColor NewColor = StreamingLevel->LevelColor;
+		TArray<FLinearColor*> ColorArray;
 		ColorArray.Add(&NewColor);
 
 		FColorPickerArgs PickerArgs;
 		PickerArgs.bIsModal = true;
 		PickerArgs.DisplayGamma = TAttribute<float>::Create( TAttribute<float>::FGetter::CreateUObject(GEngine, &UEngine::GetDisplayGamma) );
-		PickerArgs.ColorArray = &ColorArray;
+		PickerArgs.LinearColorArray = &ColorArray;
 		PickerArgs.OnColorPickerCancelled = FOnColorPickerCancelled::CreateSP(this, &FLevelViewModel::OnColorPickerCancelled);
 		PickerArgs.ParentWidget = InPickerParentWidget;
 
@@ -493,7 +493,7 @@ void FLevelViewModel::ChangeColor(const TSharedRef<SWidget>& InPickerParentWidge
 		{
 			if ( bColorPickerOK )
 			{
-				StreamingLevel->DrawColor = NewColor;
+				StreamingLevel->LevelColor = NewColor;
 				StreamingLevel->Modify();
 
 				// Update the loaded level's components so the change in color will apply immediately
diff --git a/Engine/Source/Editor/UnrealEd/Private/EditorLevelUtils.cpp b/Engine/Source/Editor/UnrealEd/Private/EditorLevelUtils.cpp
--- a/Engine/Source/Editor/UnrealEd/Private/EditorLevelUtils.cpp
+++ b/Engine/Source/Editor/UnrealEd/Private/EditorLevelUtils.cpp
@@ -197,7 +197,7 @@ namespace EditorLevelUtils
 			StreamingLevel->PackageName = LevelPackageName;
 
 			// Seed the level's draw color.
-			StreamingLevel->DrawColor = FColor::MakeRandomColor();
+			StreamingLevel->LevelColor = FLinearColor::MakeRandomColor();
 
 			// Add the new level to world.
 			InWorld->StreamingLevels.Add( StreamingLevel );
@@ -265,7 +265,7 @@ namespace EditorLevelUtils
 			NewStreamingLevel->LevelTransform = InLevel->LevelTransform;
 			NewStreamingLevel->EditorStreamingVolumes = InLevel->EditorStreamingVolumes;
 			NewStreamingLevel->MinTimeBetweenVolumeUnloadRequests = InLevel->MinTimeBetweenVolumeUnloadRequests;
-			NewStreamingLevel->DrawColor = InLevel->DrawColor;
+			NewStreamingLevel->LevelColor = InLevel->LevelColor;
 			NewStreamingLevel->Keywords = InLevel->Keywords;
 		}
 
diff --git a/Engine/Source/Editor/WorldBrowser/Private/LevelModel.cpp b/Engine/Source/Editor/WorldBrowser/Private/LevelModel.cpp
--- a/Engine/Source/Editor/WorldBrowser/Private/LevelModel.cpp
+++ b/Engine/Source/Editor/WorldBrowser/Private/LevelModel.cpp
@@ -339,14 +339,14 @@ bool FLevelModel::SupportsLevelColor() const
 	return false;
 }
 
-FColor FLevelModel::GetLevelColor() const
+FLinearColor FLevelModel::GetLevelColor() const
 {
 	// Returns Constant color, base classes will override this
 	// Currently not all base classes have the requisite support, so I've not made it pure virtual.
-	return FColor::White;
+	return FLinearColor::White;
 }
 
-void FLevelModel::SetLevelColor(FColor InColor)
+void FLevelModel::SetLevelColor(FLinearColor InColor)
 {
 	// Does nothing, base classes will override this
 	// Currently not all base classes have the requisite support, so I've not made it pure virtual.
diff --git a/Engine/Source/Editor/WorldBrowser/Private/LevelModel.h b/Engine/Source/Editor/WorldBrowser/Private/LevelModel.h
--- a/Engine/Source/Editor/WorldBrowser/Private/LevelModel.h
+++ b/Engine/Source/Editor/WorldBrowser/Private/LevelModel.h
@@ -143,10 +143,10 @@ public:
 	virtual bool SupportsLevelColor() const;
 
 	/** @return level color, used for visualization. (Show -> Advanced -> Level Coloration) */
-	virtual FColor GetLevelColor() const;
+	virtual FLinearColor GetLevelColor() const;
 
 	/** Sets level color, used for visualization. (Show -> Advanced -> Level Coloration) */
-	virtual void SetLevelColor(FColor InColor);
+	virtual void SetLevelColor(FLinearColor InColor);
 	
 	/**	@return Whether level has associated blueprint script */
 	bool HasKismet() const;
diff --git a/Engine/Source/Editor/WorldBrowser/Private/SWorldHierarchyItem.cpp b/Engine/Source/Editor/WorldBrowser/Private/SWorldHierarchyItem.cpp
--- a/Engine/Source/Editor/WorldBrowser/Private/SWorldHierarchyItem.cpp
+++ b/Engine/Source/Editor/WorldBrowser/Private/SWorldHierarchyItem.cpp
@@ -275,12 +275,12 @@ FReply SWorldHierarchyItem::OnOpenKismet()
 
 void SWorldHierarchyItem::OnSetColorFromColorPicker(FLinearColor NewColor)
 {
-	LevelModel->SetLevelColor(NewColor.ToFColor(true));
+	LevelModel->SetLevelColor(NewColor);
 }
 
 void SWorldHierarchyItem::OnColorPickerCancelled(FLinearColor OriginalColor)
 {
-	LevelModel->SetLevelColor(OriginalColor.ToFColor(true));
+	LevelModel->SetLevelColor(OriginalColor);
 }
 
 void SWorldHierarchyItem::OnColorPickerInteractiveBegin()
@@ -297,14 +297,9 @@ FReply SWorldHierarchyItem::OnChangeColor()
 {
 	if (LevelModel.IsValid())
 	{
-		// Initialize the color data for the picker window.
-		FColor NewColor = LevelModel->GetLevelColor();
-		TArray<FColor*> ColorArray;
-		ColorArray.Add(&NewColor);
-
 		FColorPickerArgs PickerArgs;
 		PickerArgs.DisplayGamma = TAttribute<float>::Create(TAttribute<float>::FGetter::CreateUObject(GEngine, &UEngine::GetDisplayGamma));
-		PickerArgs.ColorArray = &ColorArray;
+		PickerArgs.InitialColorOverride = LevelModel->GetLevelColor();
 		PickerArgs.bOnlyRefreshOnMouseUp = false;
 		PickerArgs.bOnlyRefreshOnOk = false;
 		PickerArgs.OnColorCommitted = FOnLinearColorValueChanged::CreateSP(this, &SWorldHierarchyItem::OnSetColorFromColorPicker);
diff --git a/Engine/Source/Editor/WorldBrowser/Private/StreamingLevels/StreamingLevelModel.cpp b/Engine/Source/Editor/WorldBrowser/Private/StreamingLevels/StreamingLevelModel.cpp
--- a/Engine/Source/Editor/WorldBrowser/Private/StreamingLevels/StreamingLevelModel.cpp
+++ b/Engine/Source/Editor/WorldBrowser/Private/StreamingLevels/StreamingLevelModel.cpp
@@ -69,23 +69,23 @@ bool FStreamingLevelModel::SupportsLevelColor() const
 	return LevelStreaming.IsValid();
 }
 
-FColor FStreamingLevelModel::GetLevelColor() const
+FLinearColor FStreamingLevelModel::GetLevelColor() const
 {
 	if (LevelStreaming.IsValid())
 	{
-		return LevelStreaming.Get()->DrawColor;
+		return LevelStreaming.Get()->LevelColor;
 	}
 	return FLevelModel::GetLevelColor();
 }
 
-void FStreamingLevelModel::SetLevelColor(FColor InColor)
+void FStreamingLevelModel::SetLevelColor(FLinearColor InColor)
 {
 	if (LevelStreaming.IsValid())
 	{
-		UProperty* DrawColorProperty = FindField<UProperty>(LevelStreaming->GetClass(), "DrawColor");
+		UProperty* DrawColorProperty = FindField<UProperty>(LevelStreaming->GetClass(), "LevelColor");
 		LevelStreaming->PreEditChange(DrawColorProperty);
 
-		LevelStreaming.Get()->DrawColor = InColor;
+		LevelStreaming.Get()->LevelColor = InColor;
 
 		FPropertyChangedEvent PropertyChangedEvent(DrawColorProperty, false, EPropertyChangeType::ValueSet);
 		LevelStreaming->PostEditChangeProperty(PropertyChangedEvent);
diff --git a/Engine/Source/Editor/WorldBrowser/Private/StreamingLevels/StreamingLevelModel.h b/Engine/Source/Editor/WorldBrowser/Private/StreamingLevels/StreamingLevelModel.h
--- a/Engine/Source/Editor/WorldBrowser/Private/StreamingLevels/StreamingLevelModel.h
+++ b/Engine/Source/Editor/WorldBrowser/Private/StreamingLevels/StreamingLevelModel.h
@@ -33,8 +33,8 @@ public:
 	virtual FName GetAssetName() const override;
 	virtual FName GetLongPackageName() const override;
 	virtual bool SupportsLevelColor() const override;
-	virtual FColor GetLevelColor() const override;
-	virtual void SetLevelColor(FColor InColor) override;
+	virtual FLinearColor GetLevelColor() const override;
+	virtual void SetLevelColor(FLinearColor InColor) override;
 	virtual void Update() override;
 	virtual void OnDrop(const TSharedPtr<FLevelDragDropOp>& Op) override;
 	virtual bool IsGoodToDrop(const TSharedPtr<FLevelDragDropOp>& Op) const override;
diff --git a/Engine/Source/Editor/WorldBrowser/Private/Tiles/WorldTileModel.cpp b/Engine/Source/Editor/WorldBrowser/Private/Tiles/WorldTileModel.cpp
--- a/Engine/Source/Editor/WorldBrowser/Private/Tiles/WorldTileModel.cpp
+++ b/Engine/Source/Editor/WorldBrowser/Private/Tiles/WorldTileModel.cpp
@@ -586,7 +586,7 @@ ULevelStreaming* FWorldTileModel::GetAssosiatedStreamingLevel()
 
 		//
 		AssociatedStreamingLevel->PackageName		= PackageName;
-		AssociatedStreamingLevel->DrawColor			= FColor::MakeRandomColor();
+		AssociatedStreamingLevel->LevelColor		= FLinearColor::MakeRandomColor();
 		AssociatedStreamingLevel->LevelTransform	= FTransform::Identity;
 		AssociatedStreamingLevel->PackageNameToLoad	= PackageName;
 		//
diff --git a/Engine/Source/Runtime/Core/Private/Math/Color.cpp b/Engine/Source/Runtime/Core/Private/Math/Color.cpp
--- a/Engine/Source/Runtime/Core/Private/Math/Color.cpp
+++ b/Engine/Source/Runtime/Core/Private/Math/Color.cpp
@@ -255,14 +255,18 @@ FLinearColor FLinearColor::HSVToLinearRGB() const
 						A);
 }
 
-
 /**
- * Makes a random but quite nice color.
- */
+* Makes a random but quite nice color.
+*/
+FLinearColor FLinearColor::MakeRandomColor()
+{
+	const uint8 Hue = (uint8)(FMath::FRand()*255.f);
+	return FLinearColor::FGetHSV(Hue, 0, 255);
+}
+
 FColor FColor::MakeRandomColor()
 {
-	const uint8 Hue = (uint8)( FMath::FRand()*255.f );
-	return FColor( FLinearColor::FGetHSV(Hue, 0, 255) );
+	return FColor(FLinearColor::MakeRandomColor());
 }
 
 FColor FColor::MakeRedToGreenColorFromScalar(float Scalar)
diff --git a/Engine/Source/Runtime/Core/Public/Math/Color.h b/Engine/Source/Runtime/Core/Public/Math/Color.h
--- a/Engine/Source/Runtime/Core/Public/Math/Color.h
+++ b/Engine/Source/Runtime/Core/Public/Math/Color.h
@@ -189,6 +189,11 @@ struct FLinearColor
 	 */
 	static CORE_API FLinearColor FGetHSV(uint8 H,uint8 S,uint8 V);
 
+	/**
+	* Makes a random but quite nice color.
+	*/
+	static CORE_API FLinearColor MakeRandomColor();
+
 	/**
 	 * Euclidean distance between two points.
 	 */
diff --git a/Engine/Source/Runtime/Core/Public/UObject/ObjectVersion.h b/Engine/Source/Runtime/Core/Public/UObject/ObjectVersion.h
--- a/Engine/Source/Runtime/Core/Public/UObject/ObjectVersion.h
+++ b/Engine/Source/Runtime/Core/Public/UObject/ObjectVersion.h
@@ -584,6 +584,8 @@ enum EUnrealEngineObjectUE4Version
 	VER_UE4_ADD_STRING_ASSET_REFERENCES_MAP,
 	// Apply scale from SCS RootComponent details in the Blueprint Editor to new actor instances at construction time
 	VER_UE4_BLUEPRINT_USE_SCS_ROOTCOMPONENT_SCALE,
+	// Changed level streaming to have a linear color since the visualization doesn't gamma correct.
+	VER_UE4_LEVEL_STREAMING_DRAW_COLOR_TYPE_CHANGE,
 
 	// -----<new versions can be added before this line>-------------------------------------------------
 	// - this needs to be the last line (see note below)
diff --git a/Engine/Source/Runtime/Engine/Classes/Engine/LevelStreaming.h b/Engine/Source/Runtime/Engine/Classes/Engine/LevelStreaming.h
--- a/Engine/Source/Runtime/Engine/Classes/Engine/LevelStreaming.h
+++ b/Engine/Source/Runtime/Engine/Classes/Engine/LevelStreaming.h
@@ -125,9 +125,13 @@ class ULevelStreaming : public UObject
 	UPROPERTY(EditAnywhere, AdvancedDisplay, Category=LevelStreaming)
 	uint32 bDrawOnLevelStatusMap:1;
 
+	/** Deprecated level color used for visualization. */
+	UPROPERTY()
+	FColor DrawColor_DEPRECATED;
+
 	/** The level color used for visualization. (Show -> Advanced -> Level Coloration) */
-	UPROPERTY(EditAnywhere, Category = LevelStreaming, meta=(InterpretAsLinear))
-	FColor DrawColor;
+	UPROPERTY(EditAnywhere, Category = LevelStreaming)
+	FLinearColor LevelColor;
 
 	/** The level streaming volumes bound to this level.																		*/
 	UPROPERTY(EditAnywhere, Category=LevelStreaming, meta=(DisplayName = "Streaming Volumes", NoElementDuplicate))
diff --git a/Engine/Source/Runtime/Engine/Private/BillboardComponent.cpp b/Engine/Source/Runtime/Engine/Private/BillboardComponent.cpp
--- a/Engine/Source/Runtime/Engine/Private/BillboardComponent.cpp
+++ b/Engine/Source/Runtime/Engine/Private/BillboardComponent.cpp
@@ -22,8 +22,8 @@ public:
 		, ScreenSize(InComponent->ScreenSize)
 		, U(InComponent->U)
 		, V(InComponent->V)
-		, Color(255,255,255)
-		, LevelColor(255,255,255)
+		, Color(FLinearColor::White)
+		, LevelColor(FLinearColor::White)
 		, PropertyColor(255,255,255)
 		, bIsScreenSizeScaled(InComponent->bIsScreenSizeScaled)
 #if WITH_EDITORONLY_DATA
@@ -74,11 +74,7 @@ public:
 #if WITH_EDITORONLY_DATA
 			else if (GIsEditor)
 			{
-				Color.R = 255;
-				Color.G = 255;
-				Color.B = 255;
-				// Ensure sprite is not transparent.
-				Color.A = 255;
+				Color = FLinearColor::White;
 			}
 #endif // WITH_EDITORONLY_DATA
 
@@ -94,7 +90,7 @@ public:
 			if (ULevelStreaming* LevelStreaming = FLevelUtils::FindStreamingLevel(Level))
 			{
 				// Selection takes priority over level coloration.
-				LevelColor = LevelStreaming->DrawColor;
+				LevelColor = LevelStreaming->LevelColor;
 			}
 		}
 
@@ -158,7 +154,7 @@ public:
 			{
 				ColorToUse = FColor(255,0,0);
 			}
-			FLinearColor LevelColorToUse = IsSelected() ? ColorToUse : (FLinearColor)LevelColor;
+			FLinearColor LevelColorToUse = IsSelected() ? ColorToUse : LevelColor;
 			FLinearColor PropertyColorToUse = PropertyColor;
 
 			const FLinearColor& SpriteColor = View->Family->EngineShowFlags.LevelColoration ? LevelColorToUse :
@@ -211,8 +207,8 @@ private:
 	float UL;
 	const float V;
 	float VL;
-	FColor Color;
-	FColor LevelColor;
+	FLinearColor Color;
+	FLinearColor LevelColor;
 	FColor PropertyColor;
 	const uint32 bIsScreenSizeScaled : 1;
 	uint32 bIsActorLocked : 1;
diff --git a/Engine/Source/Runtime/Engine/Private/BrushComponent.cpp b/Engine/Source/Runtime/Engine/Private/BrushComponent.cpp
--- a/Engine/Source/Runtime/Engine/Private/BrushComponent.cpp
+++ b/Engine/Source/Runtime/Engine/Private/BrushComponent.cpp
@@ -135,7 +135,7 @@ public:
 		bSolidWhenSelected(false),
 		bInManipulation(false),
 		BrushColor(GEngine->C_BrushWire),
-		LevelColor(255,255,255),
+		LevelColor(FLinearColor::White),
 		PropertyColor(255,255,255),
 		BodySetup(Component->BrushBodySetup),
 		CollisionResponse(Component->GetCollisionResponseToChannels())
@@ -170,7 +170,7 @@ public:
 				ULevelStreaming* LevelStreaming = FLevelUtils::FindStreamingLevel( Level );
 				if ( LevelStreaming )
 				{
-					LevelColor = LevelStreaming->DrawColor;
+					LevelColor = LevelStreaming->LevelColor;
 				}
 			}
 		}
@@ -427,7 +427,7 @@ private:
 	uint32 bInManipulation : 1;
 
 	FColor BrushColor;
-	FColor LevelColor;
+	FLinearColor LevelColor;
 	FColor PropertyColor;
 
 	/** Collision Response of this component**/
diff --git a/Engine/Source/Runtime/Engine/Private/Landscape/LandscapeRender.cpp b/Engine/Source/Runtime/Engine/Private/Landscape/LandscapeRender.cpp
--- a/Engine/Source/Runtime/Engine/Private/Landscape/LandscapeRender.cpp
+++ b/Engine/Source/Runtime/Engine/Private/Landscape/LandscapeRender.cpp
@@ -548,7 +548,7 @@ FLandscapeComponentSceneProxy::FLandscapeComponentSceneProxy(ULandscapeComponent
 ,	EditToolRenderData(InEditToolRenderData)
 ,	ComponentLightInfo(NULL)
 ,	LandscapeComponent(InComponent)
-,	LevelColor(1.f, 1.f, 1.f)
+,	LevelColor(FLinearColor::White)
 ,	ForcedLOD(InComponent->ForcedLOD)
 ,	LODBias(InComponent->LODBias)
 ,	LODFalloff(InComponent->GetLandscapeProxy()->LODFalloff)
@@ -647,7 +647,7 @@ FLandscapeComponentSceneProxy::FLandscapeComponentSceneProxy(ULandscapeComponent
 			ULevelStreaming* LevelStreaming = FLevelUtils::FindStreamingLevel( Level );
 			if ( LevelStreaming )
 			{
-				LevelColor = LevelStreaming->DrawColor;
+				LevelColor = LevelStreaming->LevelColor;
 			}
 		}
 	}
diff --git a/Engine/Source/Runtime/Engine/Private/LevelStreaming.cpp b/Engine/Source/Runtime/Engine/Private/LevelStreaming.cpp
--- a/Engine/Source/Runtime/Engine/Private/LevelStreaming.cpp
+++ b/Engine/Source/Runtime/Engine/Private/LevelStreaming.cpp
@@ -217,6 +217,11 @@ void ULevelStreaming::PostLoad()
 			UE_LOG(LogLevelStreaming, Display, TEXT("Invalid streaming level package name (%s). Only long package names are supported. This streaming level may not load or save properly."), *PackageName.ToString());
 		}
 	}
+
+	if (GetLinkerUE4Version() < VER_UE4_LEVEL_STREAMING_DRAW_COLOR_TYPE_CHANGE)
+	{
+		LevelColor = DrawColor_DEPRECATED;
+	}
 }
 
 UWorld* ULevelStreaming::GetWorld() const
@@ -776,9 +781,9 @@ void ULevelStreaming::PostEditChangeProperty(FPropertyChangedEvent& PropertyChan
 			}
 		}
 
-		else if (PropertyName == GET_MEMBER_NAME_CHECKED(ULevelStreaming, DrawColor))
+		else if (PropertyName == GET_MEMBER_NAME_CHECKED(ULevelStreaming, LevelColor))
 		{
-			// Make sure the level's DrawColor change is applied immediately by reregistering the
+			// Make sure the level's Level Color change is applied immediately by reregistering the
 			// components of the actor's in the level
 			if( LoadedLevel != NULL )
 			{
@@ -824,7 +829,7 @@ ULevelStreaming::ULevelStreaming(const class FPostConstructInitializeProperties&
 	: Super(PCIP)
 {
 	bShouldBeVisibleInEditor = true;
-	DrawColor = FColor(255, 255, 255, 255);
+	LevelColor = FLinearColor::White;
 	LevelTransform = FTransform::Identity;
 	MinTimeBetweenVolumeUnloadRequests = 2.0f;
 	bDrawOnLevelStatusMap = true;
diff --git a/Engine/Source/Runtime/Engine/Private/MaterialBillboardComponent.cpp b/Engine/Source/Runtime/Engine/Private/MaterialBillboardComponent.cpp
--- a/Engine/Source/Runtime/Engine/Private/MaterialBillboardComponent.cpp
+++ b/Engine/Source/Runtime/Engine/Private/MaterialBillboardComponent.cpp
@@ -83,7 +83,7 @@ public:
 	: FPrimitiveSceneProxy(InComponent)
 	, Elements(InComponent->Elements)
 	, BaseColor(FLinearColor::White)
-	, LevelColor(255,255,255)
+	, LevelColor(FLinearColor::White)
 	, PropertyColor(255,255,255)
 	{
 		AActor* Owner = InComponent->GetOwner();
@@ -95,7 +95,7 @@ public:
 			if ( LevelStreaming )
 			{
 				// Selection takes priority over level coloration.
-				LevelColor = LevelStreaming->DrawColor;
+				LevelColor = LevelStreaming->LevelColor;
 			}
 		}
 
@@ -236,7 +236,7 @@ private:
 	TArray<FMaterialSpriteElement> Elements;
 	FMaterialRelevance MaterialRelevance;
 	FColor BaseColor;
-	FColor LevelColor;
+	FLinearColor LevelColor;
 	FColor PropertyColor;
 };
 
diff --git a/Engine/Source/Runtime/Engine/Private/ModelRender.cpp b/Engine/Source/Runtime/Engine/Private/ModelRender.cpp
--- a/Engine/Source/Runtime/Engine/Private/ModelRender.cpp
+++ b/Engine/Source/Runtime/Engine/Private/ModelRender.cpp
@@ -150,7 +150,7 @@ public:
 	FModelSceneProxy(UModelComponent* InComponent):
 		FPrimitiveSceneProxy(InComponent),
 		Component(InComponent),
-		LevelColor(255,255,255),
+		LevelColor(FLinearColor::White),
 		PropertyColor(255,255,255),
 		CollisionResponse(InComponent->GetCollisionResponseToChannels())
 #if WITH_EDITOR
@@ -176,7 +176,7 @@ public:
 			ULevelStreaming* LevelStreaming = FLevelUtils::FindStreamingLevel( Level );
 			if ( LevelStreaming )
 			{
-				LevelColor = LevelStreaming->DrawColor;
+				LevelColor = LevelStreaming->LevelColor;
 			}
 		}
 
@@ -663,7 +663,7 @@ private:
 	TArray<FElementInfo> Elements;
 	TArray<FVector2D> ElementLightMapResolutions;
 
-	FColor LevelColor;
+	FLinearColor LevelColor;
 	FColor PropertyColor;
 
 	FMaterialRelevance MaterialRelevance;
diff --git a/Engine/Source/Runtime/Engine/Private/Particles/ParticleComponents.cpp b/Engine/Source/Runtime/Engine/Private/Particles/ParticleComponents.cpp
--- a/Engine/Source/Runtime/Engine/Private/Particles/ParticleComponents.cpp
+++ b/Engine/Source/Runtime/Engine/Private/Particles/ParticleComponents.cpp
@@ -470,31 +470,6 @@ bool AEmitter::GetReferencedContentObjects( TArray<UObject*>& Objects ) const
 }
 #endif
 
-
-
-//----------------------------------------------------------------------------
-
-/**
- * Try to find a level color for the specified particle system component.
- */
-static FColor* GetLevelColor(UParticleSystemComponent* Component)
-{
-	FColor* LevelColor = NULL;
-
-	AActor* Owner = Component->GetOwner();
-	if ( Owner )
-	{
-		ULevel* Level = Owner->GetLevel();
-		ULevelStreaming* LevelStreaming = FLevelUtils::FindStreamingLevel( Level );
-		if ( LevelStreaming )
-		{
-			LevelColor = &LevelStreaming->DrawColor;
-		}
-	}
-
-	return LevelColor;
-}
-
 /*-----------------------------------------------------------------------------
 	UParticleLODLevel implementation.
 -----------------------------------------------------------------------------*/
diff --git a/Engine/Source/Runtime/Engine/Private/SkeletalMesh.cpp b/Engine/Source/Runtime/Engine/Private/SkeletalMesh.cpp
--- a/Engine/Source/Runtime/Engine/Private/SkeletalMesh.cpp
+++ b/Engine/Source/Runtime/Engine/Private/SkeletalMesh.cpp
@@ -3779,7 +3779,7 @@ FSkeletalMeshSceneProxy::FSkeletalMeshSceneProxy(const USkinnedMeshComponent* Co
 		,	SkelMeshResource(InSkelMeshResource)
 		,	SkeletalMeshForDebug(Component->SkeletalMesh)
 		,	PhysicsAssetForDebug(Component->GetPhysicsAsset())
-		,	LevelColor(255,255,255)
+		,	LevelColor(FLinearColor::White)
 		,	PropertyColor(255,255,255)
 		,	bForceWireframe(Component->bForceWireframe)
 		,	bCanHighlightSelectedSections(Component->bCanHighlightSelectedSections)
@@ -3893,7 +3893,7 @@ FSkeletalMeshSceneProxy::FSkeletalMeshSceneProxy(const USkinnedMeshComponent* Co
 		ULevelStreaming* LevelStreaming = FLevelUtils::FindStreamingLevel( Level );
 		if ( LevelStreaming )
 		{
-			LevelColor = LevelStreaming->DrawColor;
+			LevelColor = LevelStreaming->LevelColor;
 		}
 	}
 
diff --git a/Engine/Source/Runtime/Engine/Private/StaticMeshRender.cpp b/Engine/Source/Runtime/Engine/Private/StaticMeshRender.cpp
--- a/Engine/Source/Runtime/Engine/Private/StaticMeshRender.cpp
+++ b/Engine/Source/Runtime/Engine/Private/StaticMeshRender.cpp
@@ -106,7 +106,7 @@ FStaticMeshSceneProxy::FStaticMeshSceneProxy(UStaticMeshComponent* InComponent):
 			ULevelStreaming* LevelStreaming = FLevelUtils::FindStreamingLevel( Level );
 			if ( LevelStreaming )
 			{
-				LevelColor = LevelStreaming->DrawColor;
+				LevelColor = LevelStreaming->LevelColor;
 			}
 		}
 
diff --git a/Engine/Source/Runtime/Engine/Public/SkeletalMeshTypes.h b/Engine/Source/Runtime/Engine/Public/SkeletalMeshTypes.h
--- a/Engine/Source/Runtime/Engine/Public/SkeletalMeshTypes.h
+++ b/Engine/Source/Runtime/Engine/Public/SkeletalMeshTypes.h
@@ -1844,7 +1844,7 @@ protected:
 	class UPhysicsAsset* PhysicsAssetForDebug;
 
 	/** data copied for rendering */
-	FColor LevelColor;
+	FLinearColor LevelColor;
 	FColor PropertyColor;
 	uint32 bForceWireframe : 1;
 	uint32 bIsCPUSkinned : 1;
