Disallow moving matinee actors to another level if their matinee is open

#ttp 335669 - EDITOR: MATINEE: With matinee open if you move a matinee actor to the same or new level, it duplicates it

#branch UE4

diff --git a/Engine/Source/Editor/Levels/Private/LevelCollectionViewModel.cpp b/Engine/Source/Editor/Levels/Private/LevelCollectionViewModel.cpp
--- a/Engine/Source/Editor/Levels/Private/LevelCollectionViewModel.cpp
+++ b/Engine/Source/Editor/Levels/Private/LevelCollectionViewModel.cpp
@@ -1108,23 +1108,21 @@ void FLevelCollectionViewModel::MoveActorsToSelected_Executed()
 			{
 				TArray<AActor*> ControlledActors;
 				InterpEditMode->MatineeActor->GetControlledActors(ControlledActors);
-				if (ControlledActors.Num() > 0)
+
+				// are any of the selected actors in the matinee
+				USelection* SelectedActors = GEditor->GetSelectedActors();
+				for (FSelectionIterator Iter(*SelectedActors); Iter; ++Iter)
 				{
-					// are any of the selected actors in the matinee
-					USelection* SelectedActors = GEditor->GetSelectedActors();
-					for (FSelectionIterator Iter(*SelectedActors); Iter; ++Iter)
+					AActor* Actor = CastChecked<AActor>(*Iter);
+					if (Actor != nullptr && (Actor == InterpEditMode->MatineeActor || ControlledActors.Contains(Actor)))
 					{
-						AActor* Actor = CastChecked<AActor>(*Iter);
-						if (Actor != nullptr && ControlledActors.Contains(Actor))
+						const bool ExitInterp = EAppReturnType::Yes == FMessageDialog::Open(EAppMsgType::YesNo, NSLOCTEXT("UnrealEd", "MatineeUnableToMove", "You must close Matinee before moving actors.\nDo you wish to do this now and continue?"));
+						if (!ExitInterp)
 						{
-							const bool ExitInterp = EAppReturnType::Yes == FMessageDialog::Open(EAppMsgType::YesNo, NSLOCTEXT("UnrealEd", "MatineeUnableToMove", "You must close Matinee before moving actors.\nDo you wish to do this now and continue?"));
-							if (!ExitInterp)
-							{
-								return;
-							}
-							GLevelEditorModeTools().DeactivateMode(FBuiltinEditorModes::EM_InterpEdit);
-							break;
+							return;
 						}
+						GLevelEditorModeTools().DeactivateMode(FBuiltinEditorModes::EM_InterpEdit);
+						break;
 					}
 				}
 			}
diff --git a/Engine/Source/Editor/WorldBrowser/Private/LevelCollectionModel.cpp b/Engine/Source/Editor/WorldBrowser/Private/LevelCollectionModel.cpp
--- a/Engine/Source/Editor/WorldBrowser/Private/LevelCollectionModel.cpp
+++ b/Engine/Source/Editor/WorldBrowser/Private/LevelCollectionModel.cpp
@@ -1467,23 +1467,21 @@ void FLevelCollectionModel::MoveActorsToSelected_Executed()
 		{
 			TArray<AActor*> ControlledActors;
 			InterpEditMode->MatineeActor->GetControlledActors(ControlledActors);
-			if (ControlledActors.Num() > 0)
+
+			// are any of the selected actors in the matinee
+			USelection* SelectedActors = GEditor->GetSelectedActors();
+			for (FSelectionIterator Iter(*SelectedActors); Iter; ++Iter)
 			{
-				// are any of the selected actors in the matinee
-				USelection* SelectedActors = GEditor->GetSelectedActors();
-				for (FSelectionIterator Iter(*SelectedActors); Iter; ++Iter)
+				AActor* Actor = CastChecked<AActor>(*Iter);
+				if (Actor != nullptr && (Actor == InterpEditMode->MatineeActor || ControlledActors.Contains(Actor)))
 				{
-					AActor* Actor = CastChecked<AActor>(*Iter);
-					if (Actor != nullptr && ControlledActors.Contains(Actor))
+					const bool ExitInterp = EAppReturnType::Yes == FMessageDialog::Open(EAppMsgType::YesNo, NSLOCTEXT("UnrealEd", "MatineeUnableToMove", "You must close Matinee before moving actors.\nDo you wish to do this now and continue?"));
+					if (!ExitInterp)
 					{
-						const bool ExitInterp = EAppReturnType::Yes == FMessageDialog::Open(EAppMsgType::YesNo, NSLOCTEXT("UnrealEd", "MatineeUnableToMove", "You must close Matinee before moving actors.\nDo you wish to do this now and continue?"));
-						if (!ExitInterp)
-						{
-							return;
-						}
-						GLevelEditorModeTools().DeactivateMode(FBuiltinEditorModes::EM_InterpEdit);
-						break;
+						return;
 					}
+					GLevelEditorModeTools().DeactivateMode(FBuiltinEditorModes::EM_InterpEdit);
+					break;
 				}
 			}
 		}
