Removed a bunch of stuff from Slate standard include, created SlateBasics.h

diff --git a/Engine/Source/Editor/Levels/Private/SLevelsViewRow.h b/Engine/Source/Editor/Levels/Private/SLevelsViewRow.h
--- a/Engine/Source/Editor/Levels/Private/SLevelsViewRow.h
+++ b/Engine/Source/Editor/Levels/Private/SLevelsViewRow.h
@@ -105,273 +105,7 @@ protected:
 	 *
 	 * @return a widget to represent the contents of a cell in this row of a TableView. 
 	 */
-	BEGIN_SLATE_FUNCTION_BUILD_OPTIMIZATION
-	virtual TSharedRef< SWidget > GenerateWidgetForColumn( const FName& ColumnID ) override
-	{
-		TSharedPtr< SWidget > TableRowContent;
-
-		if( ColumnID == LevelsView::ColumnID_LevelLabel )
-		{
-			TableRowContent =
-				SNew( SHorizontalBox )
-				+SHorizontalBox::Slot()
-				.VAlign(VAlign_Center)
-				.AutoWidth()
-				[
-					SNew( STextBlock )
-					.Font( this, &SLevelsViewRow::GetFont )
-					.Text( ViewModel.Get(), &FLevelViewModel::GetDisplayName )
-					.ColorAndOpacity( this, &SLevelsViewRow::GetColorAndOpacity )
-					.HighlightText( HighlightText )
-					.ToolTipText( LOCTEXT("DoubleClickToolTip", "Double-Click to make this the current Level") )
-					.ShadowOffset( FVector2D(1.0f, 1.0f) )
-				]
-			;
-		}
-		else if( ColumnID == LevelsView::ColumnID_ActorCount )
-		{
-			TableRowContent =
-				SNew( SHorizontalBox )
-				.Visibility(this, &SLevelsViewRow::IsActorColumnVisible )
-				+SHorizontalBox::Slot()
-				.HAlign(HAlign_Right)
-				.VAlign(VAlign_Center)
-				[
-					SNew( STextBlock )
-					.Font( this, &SLevelsViewRow::GetFont )
-					.Text( ViewModel.Get(), &FLevelViewModel::GetActorCountString )
-					.ColorAndOpacity( this, &SLevelsViewRow::GetColorAndOpacity )
-					.HighlightText( HighlightText )
-					.ToolTipText( LOCTEXT("ActorCountToolTip", "The number of Actors in this level") )
-					.ShadowOffset( FVector2D(1.0f, 1.0f) )
-				]
-			;
-		}
-		else if( ColumnID == LevelsView::ColumnID_LightmassSize )
-		{
-			TableRowContent =
-				SNew( SHorizontalBox )
-				.Visibility(this, &SLevelsViewRow::IsLightmassSizeColumnVisible )
-				+SHorizontalBox::Slot()
-				.FillWidth(1.0f)
-				.HAlign(HAlign_Right)
-				.VAlign(VAlign_Center)
-				[
-					SNew( STextBlock )
-					.Font( this, &SLevelsViewRow::GetFont )
-					.Text( ViewModel.Get(), &FLevelViewModel::GetLightmassSizeString )
-					.ColorAndOpacity( this, &SLevelsViewRow::GetColorAndOpacity )
-					.HighlightText( HighlightText )
-					.ToolTipText( LOCTEXT("LightmassSizeToolTip", "The size of the lightmap data for this level") )
-					.ShadowOffset( FVector2D(1.0f, 1.0f) )
-				]
-			;
-		}
-		else if( ColumnID == LevelsView::ColumnID_FileSize )
-		{
-			TableRowContent =
-				SNew( SHorizontalBox )
-				.Visibility(this, &SLevelsViewRow::IsFileSizeColumnVisible )
-				+SHorizontalBox::Slot()
-				.HAlign(HAlign_Right)
-				.VAlign(VAlign_Center)
-				[
-					SNew( STextBlock )
-					.Font( this, &SLevelsViewRow::GetFont )
-					.Text( ViewModel.Get(), &FLevelViewModel::GetFileSizeString )
-					.ColorAndOpacity( this, &SLevelsViewRow::GetColorAndOpacity )
-					.HighlightText( HighlightText )
-					.ToolTipText( LOCTEXT("FileSizeToolTip", "The size of the file for this level") )
-					.ShadowOffset( FVector2D(1.0f, 1.0f) )
-				]
-			;
-		}
-		else if( ColumnID == LevelsView::ColumnID_Visibility )
-		{
-			TableRowContent = 
-				SAssignNew( VisibilityButton, SButton )
-				.ContentPadding( 0 )
-				.ButtonStyle( FEditorStyle::Get(), "ToggleButton" )
-				.OnClicked( this, &SLevelsViewRow::OnToggleVisibility )
-				.ToolTipText( LOCTEXT("VisibilityButtonToolTip", "Toggle Level Visibility") )
-				.ForegroundColor( FSlateColor::UseForeground() )
-				.HAlign( HAlign_Center )
-				.VAlign( VAlign_Center )
-				.Content()
-				[
-					SNew( SImage )
-					.Image( this, &SLevelsViewRow::GetVisibilityBrushForLevel )
-					.ColorAndOpacity( this, &SLevelsViewRow::GetForegroundColorForVisibilityButton )
-				]
-			;
-		}
-		else if( ColumnID == LevelsView::ColumnID_Lock )
-		{
-			TableRowContent = 
-				SAssignNew( LockButton, SButton )
-				.ContentPadding( 0 )
-				.ButtonStyle( FEditorStyle::Get(), "ToggleButton" )
-				.OnClicked( this, &SLevelsViewRow::OnToggleLock )
-				.ToolTipText( this, &SLevelsViewRow::GetLockToolTipForLevel )
-				.ForegroundColor( FSlateColor::UseForeground() )
-				.HAlign( HAlign_Center )
-				.VAlign( VAlign_Center )
-				.Content()
-				[
-					SNew( SImage )
-					.Image( this, &SLevelsViewRow::GetLockBrushForLevel )
-					.ColorAndOpacity( this, &SLevelsViewRow::GetForegroundColorForLockButton )
-				]
-			;
-		}
-		else if( ColumnID == LevelsView::ColumnID_Color)
-		{
-			TableRowContent = 
-				SAssignNew( ColorButton, SButton )
-				.ContentPadding( 0 )
-				.ButtonStyle( FEditorStyle::Get(), "ToggleButton" )
-				.OnClicked( this, &SLevelsViewRow::OnChangeColor )
-				.ToolTipText( LOCTEXT("ColorButtonToolTip", "Change Level Color") )
-				.ForegroundColor( FSlateColor::UseForeground() )
-				.HAlign( HAlign_Center )
-				.VAlign( VAlign_Center )
-				.Content()
-				[
-					SNew( SImage )
-					.Image( this, &SLevelsViewRow::GetColorBrushForLevel )
-					.ColorAndOpacity( TAttribute< FSlateColor >::Create( TAttribute< FSlateColor >::FGetter::CreateSP( this, &SLevelsViewRow::GetLevelColorAndOpacity ) ) )
-				]
-			;
-		}
-		else if( ColumnID == LevelsView::ColumnID_SCCStatus )
-		{
-			TableRowContent = 
-				SNew(SVerticalBox)
-				+SVerticalBox::Slot()
-				.FillHeight(1.0f)
-				.HAlign(HAlign_Center)
-				.VAlign(VAlign_Center)
-				[
-					SNew(SHorizontalBox)
-					+SHorizontalBox::Slot()
-						.FillWidth(1.0f)
-						.VAlign(VAlign_Center)
-					[
-						SNew( SImage )
-							.Image( this, &SLevelsViewRow::GetSCCStateImage )
-							.ToolTipText(this, &SLevelsViewRow::GetSCCStateTooltip)
-					]
-				]
-			;
-		}
-		else if( ColumnID == LevelsView::ColumnID_Save )
-		{
-			TableRowContent = 
-				SAssignNew( SaveButton, SButton )
-				.ContentPadding( 0 )
-				.ButtonStyle( FEditorStyle::Get(), "ToggleButton" )
-				.OnClicked( this, &SLevelsViewRow::OnSave )
-				.ToolTipText( LOCTEXT("SaveButtonToolTip", "Save Level") )
-				.ForegroundColor( FSlateColor::UseForeground() )
-				.HAlign( HAlign_Center )
-				.VAlign( VAlign_Center )
-				.Content()
-				[
-					SNew( SImage )
-					.Image( this, &SLevelsViewRow::GetSaveBrushForLevel )
-					.ColorAndOpacity( TAttribute< FSlateColor >::Create( TAttribute< FSlateColor >::FGetter::CreateSP( this, &SLevelsViewRow::GetSaveButtonColorAndOpacity ) ) )
-				]
-			;
-		}
-		else if( ColumnID == LevelsView::ColumnID_Kismet )
-		{
-			TableRowContent = 
-				SAssignNew( KismetButton, SButton )
-				.ContentPadding( 0 )
-				.ButtonStyle( FEditorStyle::Get(), "ToggleButton" )
-				.OnClicked( this, &SLevelsViewRow::OnOpenKismet )
-				.ToolTipText( LOCTEXT("KismetButtonToolTip", "Open Level Blueprint") )
-				.ForegroundColor( FSlateColor::UseForeground() )
-				.HAlign( HAlign_Center )
-				.VAlign( VAlign_Center )
-				.Content()
-				[
-					SNew( SImage )
-					.Image( this, &SLevelsViewRow::GetKismetBrushForLevel )
-					.ColorAndOpacity( this, &SLevelsViewRow::GetForegroundColorForKismetButton )
-				]
-			;
-		}
-		else if( ColumnID == LevelsView::ColumnID_EditorOffset )
-		{
-			TableRowContent = 
-				SNew( SHorizontalBox )
-				.Visibility(this, &SLevelsViewRow::IsEditorOffsetColumnVisible )
-				+SHorizontalBox::Slot()
-				.Padding(FMargin(10,0,0,0))
-				.HAlign(HAlign_Fill)
-				.VAlign(VAlign_Center)
-				.FillWidth(0.25f)
-				[
-					SNew( SButton )
-					.Text( LOCTEXT("EditLevelTransform", "Viewport Edit") )
-					.ToolTipText( LOCTEXT("EditLevelToolTip", "Edit level transform in viewport.") )
-					.OnClicked( FOnClicked::CreateSP( this, &SLevelsViewRow::OnEditLevelClicked ) )
-					.IsEnabled( this, &SLevelsViewRow::LevelTransformAllowed )
-				]
-				+SHorizontalBox::Slot()
-				.Padding(FMargin(10,0,0,0))
-				.HAlign(HAlign_Fill)
-				.VAlign(VAlign_Center)
-				.FillWidth(0.55f)
-				[
-					SNew( SVectorInputBox )
-					.X( this, &SLevelsViewRow::GetLevelOffset, 0 )
-					.Y( this, &SLevelsViewRow::GetLevelOffset, 1 )
-					.Z( this, &SLevelsViewRow::GetLevelOffset, 2 )
-					.bColorAxisLabels( true )
-					.OnXCommitted( this, &SLevelsViewRow::OnSetLevelOffset, 0 )
-					.OnYCommitted( this, &SLevelsViewRow::OnSetLevelOffset, 1 )
-					.OnZCommitted( this, &SLevelsViewRow::OnSetLevelOffset, 2 )
-					.Font( FEditorStyle::GetFontStyle( TEXT("PropertyWindow.NormalFont") ) ) // Look the same as transform editing in the details panel
-					.IsEnabled( this, &SLevelsViewRow::LevelEditTextTransformAllowed )
-				]
-				+SHorizontalBox::Slot()
-				.Padding(FMargin(10,0,0,0))
-				.HAlign(HAlign_Fill)
-				.VAlign(VAlign_Center)
-				.FillWidth(0.2f)
-				[
-					SNew( SNumericEntryBox<float> )			
-					.IsEnabled( this, &SLevelsViewRow::LevelEditTextTransformAllowed )
-					.Font( FEditorStyle::GetFontStyle( TEXT("PropertyWindow.NormalFont") ) )
-					.Delta(90.0f)
-					.AllowSpin( true )
-					.MinValue(0.0f)
-					.MaxValue(270.0f)
-					.MinSliderValue(0.0f)
-					.MaxSliderValue(270.0f)
-					.Value( this, &SLevelsViewRow::GetLevelRotation )
-					.OnValueChanged( this, &SLevelsViewRow::OnSetLevelRotation )
-					.OnValueCommitted( this, &SLevelsViewRow::OnCommitLevelRotation )						
-					.OnBeginSliderMovement( this, &SLevelsViewRow::OnBeginLevelRotatonSlider )
-					.OnEndSliderMovement( this, &SLevelsViewRow::OnEndLevelRotatonSlider )
-					.LabelPadding(0)					
-					.Label()
-					[
-						SNumericEntryBox<float>::BuildLabel( LOCTEXT("LevelRotation_Label", "Yaw"), FLinearColor::White, SNumericEntryBox<float>::BlueLabelBackgroundColor )
-					]
-				]
-			;
-		}
-		else
-		{
-			checkf(false, TEXT("Unknown ColumnID provided to SLevelsView") );
-		}
-
-		return TableRowContent.ToSharedRef();
-	}
-	END_SLATE_FUNCTION_BUILD_OPTIMIZATION
+	virtual TSharedRef< SWidget > GenerateWidgetForColumn( const FName& ColumnID ) override;
 
 	/**
 	 * Called during drag and drop when the drag leaves a widget.
diff --git a/Engine/Source/Editor/MainFrame/Private/Frame/MainFrameActions.cpp b/Engine/Source/Editor/MainFrame/Private/Frame/MainFrameActions.cpp
--- a/Engine/Source/Editor/MainFrame/Private/Frame/MainFrameActions.cpp
+++ b/Engine/Source/Editor/MainFrame/Private/Frame/MainFrameActions.cpp
@@ -3,6 +3,10 @@
 
 #include "MainFramePrivatePCH.h"
 #include "MessageLog.h"
+#include "SDockTab.h"
+#include "SNotificationList.h"
+#include "NotificationManager.h"
+#include "GenericCommands.h"
 
 
 #define LOCTEXT_NAMESPACE "MainFrameActions"
diff --git a/Engine/Source/Editor/MainFrame/Private/Frame/MainFrameHandler.cpp b/Engine/Source/Editor/MainFrame/Private/Frame/MainFrameHandler.cpp
new file mode 100644
--- /dev/null
+++ b/Engine/Source/Editor/MainFrame/Private/Frame/MainFrameHandler.cpp
@@ -0,0 +1,110 @@
+// Copyright 1998-2014 Epic Games, Inc. All Rights Reserved.
+
+#include "MainFramePrivatePCH.h"
+
+#include "MainFrameHandler.h"
+#include "SDockTab.h"
+
+
+void FMainFrameHandler::ShutDownEditor()
+{
+	FEditorDelegates::OnShutdownPostPackagesSaved.Broadcast();
+
+	// Any pending autosaves should not happen.  A tick will go by before the editor shuts down and we want to avoid auto-saving during this time.
+	GUnrealEd->GetPackageAutoSaver().ResetAutoSaveTimer();
+	GEditor->RequestEndPlayMap();
+
+	// End any play on console/PC games still happening
+	GEditor->EndPlayOnLocalPc();
+
+	TSharedPtr<SWindow> RootWindow = RootWindowPtr.Pin();
+
+	// Save root window placement so we can restore it.
+	if (RootWindow.IsValid())
+	{
+		FSlateRect WindowRect = RootWindow->GetNonMaximizedRectInScreen();
+		FRootWindowLocation RootWindowLocation(FVector2D(WindowRect.Left, WindowRect.Top), WindowRect.GetSize(), RootWindow->IsWindowMaximized());
+		RootWindowLocation.SaveToIni();
+	}
+
+	// Save the visual state of the editor before we even
+	// ask whether we can shut down.
+	TSharedRef<FGlobalTabmanager> GlobalTabManager = FGlobalTabmanager::Get();
+	if (FUnrealEdMisc::Get().IsSavingLayoutOnClosedAllowed())
+	{
+		GlobalTabManager->SaveAllVisualState();
+	}
+	else
+	{
+		GConfig->EmptySection(TEXT("EditorLayouts"), *GEditorLayoutIni);
+	}
+
+	// Clear the callback for destructionfrom the main tab; otherwise it will re-enter this shutdown function.
+	if (MainTabPtr.IsValid())
+	{
+		MainTabPtr.Pin()->SetOnTabClosed(SDockTab::FOnTabClosedCallback());
+	}
+
+	// Inform the AssetEditorManager that the editor is exiting so that it may save open assets
+	// and report usage stats
+	FAssetEditorManager::Get().OnExit();
+
+	if (RootWindow.IsValid())
+	{
+		RootWindow->SetRequestDestroyWindowOverride(FRequestDestroyWindowOverride());
+		RootWindow->RequestDestroyWindow();
+	}
+
+	// Save out any config settings for the editor so they don't get lost
+	GEditor->SaveConfig();
+	GLevelEditorModeTools().SaveConfig();
+
+	// Delete user settings, if requested
+	if (FUnrealEdMisc::Get().IsDeletePreferences())
+	{
+		IFileManager::Get().Delete(*GEditorUserSettingsIni);
+	}
+
+	// Take a screenshot of this project for the project browser
+	if (FApp::HasGameName())
+	{
+		const FString ExistingBaseFilename = FString(FApp::GetGameName()) + TEXT(".png");
+		const FString ExistingScreenshotFilename = FPaths::Combine(*FPaths::GameDir(), *ExistingBaseFilename);
+
+		// If there is already a screenshot, no need to take an auto screenshot
+		if (!FPaths::FileExists(ExistingScreenshotFilename))
+		{
+			const FString ScreenShotFilename = FPaths::Combine(*FPaths::GameSavedDir(), TEXT("AutoScreenshot.png"));
+			FViewport* Viewport = GEditor->GetActiveViewport();
+			if (Viewport)
+			{
+				UThumbnailManager::CaptureProjectThumbnail(Viewport, ScreenShotFilename, false);
+			}
+		}
+	}
+
+	// Shut down the editor
+	// NOTE: We can't close the editor from within this stack frame as it will cause various DLLs
+	//       (such as MainFrame) to become unloaded out from underneath the code pointer.  We'll shut down
+	//       as soon as it's safe to do so.
+	// Note this is the only place in slate that should be calling QUIT_EDITOR
+	GEngine->DeferredCommands.Add(TEXT("QUIT_EDITOR"));
+}
+
+void FMainFrameHandler::EnableTabClosedDelegate()
+{
+	if (MainTabPtr.IsValid())
+	{
+		MainTabPtr.Pin()->SetOnTabClosed(SDockTab::FOnTabClosedCallback::CreateRaw(this, &FMainFrameHandler::ShutDownEditor));
+		MainTabPtr.Pin()->SetCanCloseTab(SDockTab::FCanCloseTab::CreateRaw(this, &FMainFrameHandler::CanCloseTab));
+	}
+}
+
+void FMainFrameHandler::DisableTabClosedDelegate()
+{
+	if (MainTabPtr.IsValid())
+	{
+		MainTabPtr.Pin()->SetOnTabClosed(SDockTab::FOnTabClosedCallback());
+		MainTabPtr.Pin()->SetCanCloseTab(SDockTab::FCanCloseTab());
+	}
+}
diff --git a/Engine/Source/Editor/MainFrame/Private/Frame/MainFrameHandler.h b/Engine/Source/Editor/MainFrame/Private/Frame/MainFrameHandler.h
--- a/Engine/Source/Editor/MainFrame/Private/Frame/MainFrameHandler.h
+++ b/Engine/Source/Editor/MainFrame/Private/Frame/MainFrameHandler.h
@@ -21,90 +21,7 @@ public:
 	/**
 	 * Shuts down the Editor.
 	 */
-	void ShutDownEditor( )
-	{
-		FEditorDelegates::OnShutdownPostPackagesSaved.Broadcast();
-
-		// Any pending autosaves should not happen.  A tick will go by before the editor shuts down and we want to avoid auto-saving during this time.
-		GUnrealEd->GetPackageAutoSaver().ResetAutoSaveTimer();
-		GEditor->RequestEndPlayMap();
-
-		// End any play on console/PC games still happening
-		GEditor->EndPlayOnLocalPc();
-
-		TSharedPtr<SWindow> RootWindow = RootWindowPtr.Pin();
-
-		// Save root window placement so we can restore it.
-		if( RootWindow.IsValid() )
-		{
-			FSlateRect WindowRect = RootWindow->GetNonMaximizedRectInScreen();
-			FRootWindowLocation RootWindowLocation( FVector2D(WindowRect.Left, WindowRect.Top), WindowRect.GetSize(), RootWindow->IsWindowMaximized() );
-			RootWindowLocation.SaveToIni();
-		}
-
-		// Save the visual state of the editor before we even
-		// ask whether we can shut down.
-		TSharedRef<FGlobalTabmanager> GlobalTabManager = FGlobalTabmanager::Get();
-		if(FUnrealEdMisc::Get().IsSavingLayoutOnClosedAllowed())
-		{
-			GlobalTabManager->SaveAllVisualState();
-		}
-		else
-		{
-			GConfig->EmptySection(TEXT("EditorLayouts"), *GEditorLayoutIni);
-		}
-
-		// Clear the callback for destructionfrom the main tab; otherwise it will re-enter this shutdown function.
-		if ( MainTabPtr.IsValid() )
-		{
-			MainTabPtr.Pin()->SetOnTabClosed( SDockTab::FOnTabClosedCallback() );
-		}
-
-		// Inform the AssetEditorManager that the editor is exiting so that it may save open assets
-		// and report usage stats
-		FAssetEditorManager::Get().OnExit();
-		
-		if (RootWindow.IsValid())
-		{
-			RootWindow->SetRequestDestroyWindowOverride(FRequestDestroyWindowOverride());
-			RootWindow->RequestDestroyWindow();
-		}
-
-		// Save out any config settings for the editor so they don't get lost
-		GEditor->SaveConfig();
-		GLevelEditorModeTools().SaveConfig();
-
-		// Delete user settings, if requested
-		if (FUnrealEdMisc::Get().IsDeletePreferences())
-		{
-			IFileManager::Get().Delete(*GEditorUserSettingsIni);
-		}
-
-		// Take a screenshot of this project for the project browser
-		if ( FApp::HasGameName() )
-		{
-			const FString ExistingBaseFilename = FString(FApp::GetGameName()) + TEXT(".png");
-			const FString ExistingScreenshotFilename = FPaths::Combine(*FPaths::GameDir(), *ExistingBaseFilename);
-
-			// If there is already a screenshot, no need to take an auto screenshot
-			if ( !FPaths::FileExists(ExistingScreenshotFilename) )
-			{
-				const FString ScreenShotFilename = FPaths::Combine(*FPaths::GameSavedDir(), TEXT("AutoScreenshot.png"));
-				FViewport* Viewport = GEditor->GetActiveViewport();
-				if ( Viewport )
-				{
-					UThumbnailManager::CaptureProjectThumbnail(Viewport, ScreenShotFilename, false);
-				}
-			}
-		}
-
-		// Shut down the editor
-		// NOTE: We can't close the editor from within this stack frame as it will cause various DLLs
-		//       (such as MainFrame) to become unloaded out from underneath the code pointer.  We'll shut down
-		//       as soon as it's safe to do so.
-		// Note this is the only place in slate that should be calling QUIT_EDITOR
-		GEngine->DeferredCommands.Add( TEXT( "QUIT_EDITOR" ) );
-	}
+	void ShutDownEditor( );
 	
 	/**
 	 * Checks whether the main frame tab can be closed.
@@ -305,24 +222,10 @@ public:
 	}
 	
 	/** Enables the delegate responsible for shutting down the editor when the main tab is closed */
-	void EnableTabClosedDelegate()
-	{
-		if ( MainTabPtr.IsValid() )
-		{
-			MainTabPtr.Pin()->SetOnTabClosed(SDockTab::FOnTabClosedCallback::CreateRaw(this, &FMainFrameHandler::ShutDownEditor));
-			MainTabPtr.Pin()->SetCanCloseTab(SDockTab::FCanCloseTab::CreateRaw(this, &FMainFrameHandler::CanCloseTab));
-		}
-	}
+	void EnableTabClosedDelegate();
 
 	/** Disables the delegate responsible for shutting down the editor when the main tab is closed */
-	void DisableTabClosedDelegate()
-	{
-		if ( MainTabPtr.IsValid() )
-		{
-			MainTabPtr.Pin()->SetOnTabClosed(SDockTab::FOnTabClosedCallback());
-			MainTabPtr.Pin()->SetCanCloseTab(SDockTab::FCanCloseTab());
-		}
-	}
+	void DisableTabClosedDelegate();
 
 private:
 
diff --git a/Engine/Source/Editor/Matinee/Private/Matinee.cpp b/Engine/Source/Editor/Matinee/Private/Matinee.cpp
--- a/Engine/Source/Editor/Matinee/Private/Matinee.cpp
+++ b/Engine/Source/Editor/Matinee/Private/Matinee.cpp
@@ -39,6 +39,9 @@
 #include "EngineAnalytics.h"
 
 #include "SColorPicker.h"
+#include "SDockTab.h"
+#include "STextComboBox.h"
+#include "GenericCommands.h"
 
 
 DEFINE_LOG_CATEGORY(LogSlateMatinee);
diff --git a/Engine/Source/Editor/Matinee/Private/MatineeKeyReduction.cpp b/Engine/Source/Editor/Matinee/Private/MatineeKeyReduction.cpp
--- a/Engine/Source/Editor/Matinee/Private/MatineeKeyReduction.cpp
+++ b/Engine/Source/Editor/Matinee/Private/MatineeKeyReduction.cpp
@@ -4,8 +4,10 @@
 #include "Matinee/MatineeActor.h"
 #include "Matinee/InterpGroupInst.h"
 #include "Matinee.h"
-#include "Slate.h"
+#include "SlateBasics.h"
 #include "SNumericEntryBox.h"
+#include "SNotificationList.h"
+#include "NotificationManager.h"
 
 
 /**
diff --git a/Engine/Source/Editor/Matinee/Private/MatineeMenus.cpp b/Engine/Source/Editor/Matinee/Private/MatineeMenus.cpp
--- a/Engine/Source/Editor/Matinee/Private/MatineeMenus.cpp
+++ b/Engine/Source/Editor/Matinee/Private/MatineeMenus.cpp
@@ -46,6 +46,12 @@
 #include "EngineAnalytics.h"
 
 #include "SColorPicker.h"
+#include "SDockTab.h"
+#include "STextEntryPopup.h"
+#include "SNotificationList.h"
+#include "NotificationManager.h"
+#include "SNumericEntryBox.h"
+#include "STextComboPopup.h"
 
 
 #define LOCTEXT_NAMESPACE "MatineeMenus"
diff --git a/Engine/Source/Editor/Matinee/Private/MatineeRecordMovie.cpp b/Engine/Source/Editor/Matinee/Private/MatineeRecordMovie.cpp
--- a/Engine/Source/Editor/Matinee/Private/MatineeRecordMovie.cpp
+++ b/Engine/Source/Editor/Matinee/Private/MatineeRecordMovie.cpp
@@ -3,6 +3,8 @@
 #include "MatineeModule.h"
 #include "Matinee/MatineeActor.h"
 #include "Matinee.h"
+#include "STextComboBox.h"
+#include "SNumericEntryBox.h"
 
 
 #define LOCTEXT_NAMESPACE "MatineeMovieCapture"
diff --git a/Engine/Source/Editor/Matinee/Private/MatineeToolbars.cpp b/Engine/Source/Editor/Matinee/Private/MatineeToolbars.cpp
--- a/Engine/Source/Editor/Matinee/Private/MatineeToolbars.cpp
+++ b/Engine/Source/Editor/Matinee/Private/MatineeToolbars.cpp
@@ -3,7 +3,8 @@
 #include "MatineeModule.h"
 #include "Matinee.h"
 #include "MatineeActions.h"
-#include "Slate.h"
+#include "SlateBasics.h"
+#include "STextComboBox.h"
 
 
 FString FMatinee::GetToolbarSnapText()
diff --git a/Engine/Source/Editor/Matinee/Private/MatineeTrackHelpers.cpp b/Engine/Source/Editor/Matinee/Private/MatineeTrackHelpers.cpp
--- a/Engine/Source/Editor/Matinee/Private/MatineeTrackHelpers.cpp
+++ b/Engine/Source/Editor/Matinee/Private/MatineeTrackHelpers.cpp
@@ -22,6 +22,8 @@
 #include "Matinee/InterpTrackLinearColorProp.h"
 #include "Matinee/InterpGroupInst.h"
 #include "Matinee/MatineeAnimInterface.h"
+#include "STextComboPopup.h"
+#include "STextEntryPopup.h"
 
 #define LOCTEXT_NAMESPACE "MatineeTrackHelpers"
 
diff --git a/Engine/Source/Editor/NewsFeed/Private/Widgets/SNewsFeed.cpp b/Engine/Source/Editor/NewsFeed/Private/Widgets/SNewsFeed.cpp
--- a/Engine/Source/Editor/NewsFeed/Private/Widgets/SNewsFeed.cpp
+++ b/Engine/Source/Editor/NewsFeed/Private/Widgets/SNewsFeed.cpp
@@ -1,6 +1,8 @@
 // Copyright 1998-2014 Epic Games, Inc. All Rights Reserved.
 
 #include "NewsFeedPrivatePCH.h"
+#include "SThrobber.h"
+#include "SHyperlink.h"
 
 
 #define LOCTEXT_NAMESPACE "SNewsFeed"
diff --git a/Engine/Source/Editor/NiagaraEditor/Private/SNiagaraEffectEditorWidget.cpp b/Engine/Source/Editor/NiagaraEditor/Private/SNiagaraEffectEditorWidget.cpp
new file mode 100644
--- /dev/null
+++ b/Engine/Source/Editor/NiagaraEditor/Private/SNiagaraEffectEditorWidget.cpp
@@ -0,0 +1,155 @@
+// Copyright 1998-2014 Epic Games, Inc. All Rights Reserved.
+
+#include "NiagaraEditorPrivatePCH.h"
+#include "SNiagaraEffectEditorWidget.h"
+#include "SNotificationList.h"
+
+void SNiagaraEffectEditorWidget::Construct(const FArguments& InArgs)
+{
+	FString NiagaraEffectText(TEXT("Niagara Effect"));
+	FString ReadOnlyText(TEXT("Read Only"));
+	/*
+		Commands = MakeShareable(new FUICommandList());
+		IsEditable = InArgs._IsEditable;
+		Appearance = InArgs._Appearance;
+		TitleBarEnabledOnly = InArgs._TitleBarEnabledOnly;
+		TitleBar = InArgs._TitleBar;
+		bAutoExpandActionMenu = InArgs._AutoExpandActionMenu;
+		ShowGraphStateOverlay = InArgs._ShowGraphStateOverlay;
+		OnNavigateHistoryBack = InArgs._OnNavigateHistoryBack;
+		OnNavigateHistoryForward = InArgs._OnNavigateHistoryForward;
+		OnNodeSpawnedByKeymap = InArgs._GraphEvents.OnNodeSpawnedByKeymap;
+		*/
+
+	/*
+	// Make sure that the editor knows about what kinds
+	// of commands GraphEditor can do.
+	FGraphEditorCommands::Register();
+
+	// Tell GraphEditor how to handle all the known commands
+	{
+	Commands->MapAction(FGraphEditorCommands::Get().ReconstructNodes,
+	FExecuteAction::CreateSP(this, &SGraphEditorImpl::ReconstructNodes),
+	FCanExecuteAction::CreateSP(this, &SGraphEditorImpl::CanReconstructNodes)
+	);
+
+	Commands->MapAction(FGraphEditorCommands::Get().BreakNodeLinks,
+	FExecuteAction::CreateSP(this, &SGraphEditorImpl::BreakNodeLinks),
+	FCanExecuteAction::CreateSP(this, &SGraphEditorImpl::CanBreakNodeLinks)
+	);
+
+	Commands->MapAction(FGraphEditorCommands::Get().BreakPinLinks,
+	FExecuteAction::CreateSP(this, &SGraphEditorImpl::BreakPinLinks, true),
+	FCanExecuteAction::CreateSP(this, &SGraphEditorImpl::CanBreakPinLinks)
+	);
+
+	// Append any additional commands that a consumer of GraphEditor wants us to be aware of.
+	const TSharedPtr<FUICommandList>& AdditionalCommands = InArgs._AdditionalCommands;
+	if (AdditionalCommands.IsValid())
+	{
+	Commands->Append(AdditionalCommands.ToSharedRef());
+	}
+
+	}
+	*/
+
+	EffectObj = InArgs._EffectObj;
+	//bNeedsRefresh = false;
+
+
+	FPropertyEditorModule& PropertyEditorModule = FModuleManager::GetModuleChecked<FPropertyEditorModule>("PropertyEditor");
+	const FDetailsViewArgs DetailsViewArgs(false, false, true, false, true, this);
+	EmitterDetails = PropertyEditorModule.CreateDetailView(DetailsViewArgs);
+	//FOnGetDetailCustomizationInstance LayoutMICDetails = FOnGetDetailCustomizationInstance::CreateStatic(&FMaterialInstanceParameterDetails::MakeInstance, MaterialEditorInstance, FGetShowHiddenParameters::CreateSP(this, &FMaterialInstanceEditor::GetShowHiddenParameters));
+	//EmitterDetails->RegisterInstancedCustomPropertyLayout(UMaterialEditorInstanceConstant::StaticClass(), LayoutMICDetails);
+
+	//EmitterDetails->SetObjects(EmitterObj, true);
+
+
+	TSharedPtr<SOverlay> OverlayWidget;
+	this->ChildSlot
+		[
+			SAssignNew(OverlayWidget, SOverlay)
+
+			// details view on the left
+			+ SOverlay::Slot()
+			.Expose(DetailsViewSlot)
+			.HAlign(HAlign_Left)
+			.VAlign(VAlign_Top)
+			[
+				SNew(SBorder)
+				.Padding(4)
+				[
+					EmitterDetails.ToSharedRef()
+				]
+			]
+
+			/*
+			// The graph panel
+			+ SOverlay::Slot()
+			.Expose(GraphPanelSlot)
+			[
+			SAssignNew(GraphPanel, SGraphPanel)
+			.EffectObj(EffectObj)
+			.GraphObjToDiff(InArgs._GraphToDiff)
+			//					.OnGetContextMenuFor(this, &SGraphEditorImpl::GraphEd_OnGetContextMenuFor)
+			.OnSelectionChanged(InArgs._GraphEvents.OnSelectionChanged)
+			.OnNodeDoubleClicked(InArgs._GraphEvents.OnNodeDoubleClicked)
+			//					.IsEditable(this, &SGraphEditorImpl::IsGraphEditable)
+			.OnDropActor(InArgs._GraphEvents.OnDropActor)
+			.OnDropStreamingLevel(InArgs._GraphEvents.OnDropStreamingLevel)
+			//					.IsEnabled(this, &SGraphEditorImpl::GraphEd_OnGetGraphEnabled)
+			.OnVerifyTextCommit(InArgs._GraphEvents.OnVerifyTextCommit)
+			.OnTextCommitted(InArgs._GraphEvents.OnTextCommitted)
+			.OnSpawnNodeByShortcut(InArgs._GraphEvents.OnSpawnNodeByShortcut)
+			//					.OnUpdateGraphPanel(this, &SGraphEditorImpl::GraphEd_OnPanelUpdated)
+			.OnDisallowedPinConnection(InArgs._GraphEvents.OnDisallowedPinConnection)
+			.ShowGraphStateOverlay(InArgs._ShowGraphStateOverlay)
+			]*/
+
+			// Title bar - optional
+			+SOverlay::Slot()
+				.VAlign(VAlign_Top)
+				[
+					InArgs._TitleBar.IsValid() ? InArgs._TitleBar.ToSharedRef() : (TSharedRef<SWidget>)SNullWidget::NullWidget
+				]
+
+			// Bottom-right corner text indicating the type of tool
+			+ SOverlay::Slot()
+				.Padding(10)
+				.VAlign(VAlign_Bottom)
+				.HAlign(HAlign_Right)
+				[
+					SNew(STextBlock)
+					.Visibility(EVisibility::HitTestInvisible)
+					.TextStyle(FEditorStyle::Get(), "Graph.CornerText")
+					.Text(NiagaraEffectText)
+				]
+
+			// Top-right corner text indicating read only when not simulating
+			+ SOverlay::Slot()
+				.Padding(20)
+				.VAlign(VAlign_Top)
+				.HAlign(HAlign_Right)
+				[
+					SNew(STextBlock)
+					.Visibility(this, &SNiagaraEffectEditorWidget::ReadOnlyVisibility)
+					.TextStyle(FEditorStyle::Get(), "Graph.CornerText")
+					.Text(ReadOnlyText)
+				]
+
+			// Bottom-right corner text for notification list position
+			+ SOverlay::Slot()
+				.Padding(15.f)
+				.VAlign(VAlign_Bottom)
+				.HAlign(HAlign_Right)
+				[
+					SAssignNew(NotificationListPtr, SNotificationList)
+					.Visibility(EVisibility::HitTestInvisible)
+				]
+		];
+
+	//GraphPanel->RestoreViewSettings(FVector2D::ZeroVector, -1);
+
+	//NotifyGraphChanged();
+}
diff --git a/Engine/Source/Editor/NiagaraEditor/Private/SNiagaraEffectEditorWidget.h b/Engine/Source/Editor/NiagaraEditor/Private/SNiagaraEffectEditorWidget.h
--- a/Engine/Source/Editor/NiagaraEditor/Private/SNiagaraEffectEditorWidget.h
+++ b/Engine/Source/Editor/NiagaraEditor/Private/SNiagaraEffectEditorWidget.h
@@ -1,14 +1,13 @@
 // Copyright 1998-2014 Epic Games, Inc. All Rights Reserved.
 
+#pragma once
+
 #include "Visibility.h"
-#include "Slate.h"
+#include "SlateBasics.h"
 #include "Editor/PropertyEditor/Public/PropertyEditorModule.h"
 #include "Editor/PropertyEditor/Public/IDetailsView.h"
 
 
-FString NiagaraEffectText(TEXT("Niagara Effect"));
-FString ReadOnlyText(TEXT("Read Only"));
-
 class SNiagaraEffectEditorWidget : public SCompoundWidget, public FNotifyHook
 {
 private:
@@ -33,152 +32,6 @@ public:
 		return EVisibility::Hidden;
 	}
 
-	void Construct(const FArguments& InArgs)
-	{
-		/*
-		Commands = MakeShareable(new FUICommandList());
-		IsEditable = InArgs._IsEditable;
-		Appearance = InArgs._Appearance;
-		TitleBarEnabledOnly = InArgs._TitleBarEnabledOnly;
-		TitleBar = InArgs._TitleBar;
-		bAutoExpandActionMenu = InArgs._AutoExpandActionMenu;
-		ShowGraphStateOverlay = InArgs._ShowGraphStateOverlay;
-		OnNavigateHistoryBack = InArgs._OnNavigateHistoryBack;
-		OnNavigateHistoryForward = InArgs._OnNavigateHistoryForward;
-		OnNodeSpawnedByKeymap = InArgs._GraphEvents.OnNodeSpawnedByKeymap;
-		*/
-
-		/*
-		// Make sure that the editor knows about what kinds
-		// of commands GraphEditor can do.
-		FGraphEditorCommands::Register();
-
-		// Tell GraphEditor how to handle all the known commands
-		{
-		Commands->MapAction(FGraphEditorCommands::Get().ReconstructNodes,
-		FExecuteAction::CreateSP(this, &SGraphEditorImpl::ReconstructNodes),
-		FCanExecuteAction::CreateSP(this, &SGraphEditorImpl::CanReconstructNodes)
-		);
-
-		Commands->MapAction(FGraphEditorCommands::Get().BreakNodeLinks,
-		FExecuteAction::CreateSP(this, &SGraphEditorImpl::BreakNodeLinks),
-		FCanExecuteAction::CreateSP(this, &SGraphEditorImpl::CanBreakNodeLinks)
-		);
-
-		Commands->MapAction(FGraphEditorCommands::Get().BreakPinLinks,
-		FExecuteAction::CreateSP(this, &SGraphEditorImpl::BreakPinLinks, true),
-		FCanExecuteAction::CreateSP(this, &SGraphEditorImpl::CanBreakPinLinks)
-		);
-
-		// Append any additional commands that a consumer of GraphEditor wants us to be aware of.
-		const TSharedPtr<FUICommandList>& AdditionalCommands = InArgs._AdditionalCommands;
-		if (AdditionalCommands.IsValid())
-		{
-		Commands->Append(AdditionalCommands.ToSharedRef());
-		}
-
-		}
-		*/
-
-		EffectObj = InArgs._EffectObj;
-		//bNeedsRefresh = false;
-
-
-		FPropertyEditorModule& PropertyEditorModule = FModuleManager::GetModuleChecked<FPropertyEditorModule>("PropertyEditor");
-		const FDetailsViewArgs DetailsViewArgs(false, false, true, false, true, this);
-		EmitterDetails = PropertyEditorModule.CreateDetailView(DetailsViewArgs);
-		//FOnGetDetailCustomizationInstance LayoutMICDetails = FOnGetDetailCustomizationInstance::CreateStatic(&FMaterialInstanceParameterDetails::MakeInstance, MaterialEditorInstance, FGetShowHiddenParameters::CreateSP(this, &FMaterialInstanceEditor::GetShowHiddenParameters));
-		//EmitterDetails->RegisterInstancedCustomPropertyLayout(UMaterialEditorInstanceConstant::StaticClass(), LayoutMICDetails);
-
-		//EmitterDetails->SetObjects(EmitterObj, true);
-
-
-		TSharedPtr<SOverlay> OverlayWidget;
-		this->ChildSlot
-			[
-				SAssignNew(OverlayWidget, SOverlay)
-
-				// details view on the left
-				+ SOverlay::Slot()
-				.Expose(DetailsViewSlot)
-				.HAlign(HAlign_Left)
-				.VAlign(VAlign_Top)
-				[
-					SNew(SBorder)
-					.Padding(4)
-					[
-						EmitterDetails.ToSharedRef()
-					]
-				]
-
-				/*
-				// The graph panel
-				+ SOverlay::Slot()
-				.Expose(GraphPanelSlot)
-				[
-					SAssignNew(GraphPanel, SGraphPanel)
-					.EffectObj(EffectObj)
-					.GraphObjToDiff(InArgs._GraphToDiff)
-					//					.OnGetContextMenuFor(this, &SGraphEditorImpl::GraphEd_OnGetContextMenuFor)
-					.OnSelectionChanged(InArgs._GraphEvents.OnSelectionChanged)
-					.OnNodeDoubleClicked(InArgs._GraphEvents.OnNodeDoubleClicked)
-					//					.IsEditable(this, &SGraphEditorImpl::IsGraphEditable)
-					.OnDropActor(InArgs._GraphEvents.OnDropActor)
-					.OnDropStreamingLevel(InArgs._GraphEvents.OnDropStreamingLevel)
-					//					.IsEnabled(this, &SGraphEditorImpl::GraphEd_OnGetGraphEnabled)
-					.OnVerifyTextCommit(InArgs._GraphEvents.OnVerifyTextCommit)
-					.OnTextCommitted(InArgs._GraphEvents.OnTextCommitted)
-					.OnSpawnNodeByShortcut(InArgs._GraphEvents.OnSpawnNodeByShortcut)
-					//					.OnUpdateGraphPanel(this, &SGraphEditorImpl::GraphEd_OnPanelUpdated)
-					.OnDisallowedPinConnection(InArgs._GraphEvents.OnDisallowedPinConnection)
-					.ShowGraphStateOverlay(InArgs._ShowGraphStateOverlay)
-				]*/
-
-				// Title bar - optional
-				+ SOverlay::Slot()
-					.VAlign(VAlign_Top)
-					[
-						InArgs._TitleBar.IsValid() ? InArgs._TitleBar.ToSharedRef() : (TSharedRef<SWidget>)SNullWidget::NullWidget
-					]
-
-				// Bottom-right corner text indicating the type of tool
-				+ SOverlay::Slot()
-					.Padding(10)
-					.VAlign(VAlign_Bottom)
-					.HAlign(HAlign_Right)
-					[
-						SNew(STextBlock)
-						.Visibility(EVisibility::HitTestInvisible)
-						.TextStyle(FEditorStyle::Get(), "Graph.CornerText")
-						.Text( NiagaraEffectText )
-					]
-
-				// Top-right corner text indicating read only when not simulating
-				+ SOverlay::Slot()
-					.Padding(20)
-					.VAlign(VAlign_Top)
-					.HAlign(HAlign_Right)
-					[
-						SNew(STextBlock)
-						.Visibility(this, &SNiagaraEffectEditorWidget::ReadOnlyVisibility)
-						.TextStyle(FEditorStyle::Get(), "Graph.CornerText")
-						.Text(ReadOnlyText)
-					]
-
-				// Bottom-right corner text for notification list position
-				+ SOverlay::Slot()
-					.Padding(15.f)
-					.VAlign(VAlign_Bottom)
-					.HAlign(HAlign_Right)
-					[
-						SAssignNew(NotificationListPtr, SNotificationList)
-						.Visibility(EVisibility::HitTestInvisible)
-					]
-			];
-
-		//GraphPanel->RestoreViewSettings(FVector2D::ZeroVector, -1);
-
-		//NotifyGraphChanged();
-	}
+	void Construct(const FArguments& InArgs);
 };
 
diff --git a/Engine/Source/Editor/PListEditor/Private/PListEditor.cpp b/Engine/Source/Editor/PListEditor/Private/PListEditor.cpp
--- a/Engine/Source/Editor/PListEditor/Private/PListEditor.cpp
+++ b/Engine/Source/Editor/PListEditor/Private/PListEditor.cpp
@@ -2,10 +2,11 @@
 
 #include "UnrealEd.h"
 
-#include "Slate.h"
+#include "SlateBasics.h"
 #include "PListEditor.h"
 #include "SPlistEditor.h"
 #include "Editor/WorkspaceMenuStructure/Public/WorkspaceMenuStructureModule.h"
+#include "SDockTab.h"
 
 IMPLEMENT_MODULE( FPListEditor, PListEditor );
 
diff --git a/Engine/Source/Editor/PListEditor/Private/SPlistEditor.cpp b/Engine/Source/Editor/PListEditor/Private/SPlistEditor.cpp
--- a/Engine/Source/Editor/PListEditor/Private/SPlistEditor.cpp
+++ b/Engine/Source/Editor/PListEditor/Private/SPlistEditor.cpp
@@ -6,6 +6,9 @@
 #include "DesktopPlatformModule.h"
 #include "MainFrame.h"
 
+#include "SSearchBox.h"
+#include "SExpandableArea.h"
+#include "SNotificationList.h"
 #define LOCTEXT_NAMESPACE "PListEditor"
 
 /** Ticks any internals of the widget as necessary. */
diff --git a/Engine/Source/Editor/Persona/Private/Customization/AnimGraphNodeSlotDetails.cpp b/Engine/Source/Editor/Persona/Private/Customization/AnimGraphNodeSlotDetails.cpp
--- a/Engine/Source/Editor/Persona/Private/Customization/AnimGraphNodeSlotDetails.cpp
+++ b/Engine/Source/Editor/Persona/Private/Customization/AnimGraphNodeSlotDetails.cpp
@@ -11,6 +11,8 @@
 #include "ScopedTransaction.h"
 #include "AnimGraphNodeSlotDetails.h"
 #include "Persona.h"
+#include "STextComboBox.h"
+#include "STextEntryPopup.h"
 
 #define LOCTEXT_NAMESPACE "AnimNodeSlotDetails"
 
diff --git a/Engine/Source/Editor/Persona/Private/Persona.cpp b/Engine/Source/Editor/Persona/Private/Persona.cpp
--- a/Engine/Source/Editor/Persona/Private/Persona.cpp
+++ b/Engine/Source/Editor/Persona/Private/Persona.cpp
@@ -63,6 +63,10 @@
 #include "FbxMeshUtils.h"
 #include "AnimationEditorUtils.h"
 #include "SAnimationSequenceBrowser.h"
+#include "SDockTab.h"
+#include "GenericCommands.h"
+#include "SNotificationList.h"
+#include "NotificationManager.h"
 
 #define LOCTEXT_NAMESPACE "FPersona"
 
diff --git a/Engine/Source/Editor/Persona/Private/SAnimCurvePanel.cpp b/Engine/Source/Editor/Persona/Private/SAnimCurvePanel.cpp
--- a/Engine/Source/Editor/Persona/Private/SAnimCurvePanel.cpp
+++ b/Engine/Source/Editor/Persona/Private/SAnimCurvePanel.cpp
@@ -9,6 +9,8 @@
 #include "Editor/KismetWidgets/Public/SScrubWidget.h"
 #include "AssetRegistryModule.h"
 #include "Kismet2NameValidators.h"
+#include "SExpandableArea.h"
+#include "STextEntryPopup.h"
 
 #define LOCTEXT_NAMESPACE "AnimCurvePanel"
 
diff --git a/Engine/Source/Editor/Persona/Private/SAnimMontagePanel.cpp b/Engine/Source/Editor/Persona/Private/SAnimMontagePanel.cpp
--- a/Engine/Source/Editor/Persona/Private/SAnimMontagePanel.cpp
+++ b/Engine/Source/Editor/Persona/Private/SAnimMontagePanel.cpp
@@ -10,6 +10,8 @@
 #include "SAnimSegmentsPanel.h"
 #include "SMontageEditor.h"
 #include "Editor/UnrealEd/Public/DragAndDrop/AssetDragDropOp.h"
+#include "SExpandableArea.h"
+#include "STextEntryPopup.h"
 
 #define LOCTEXT_NAMESPACE "AnimMontagePanel"
 
diff --git a/Engine/Source/Editor/Persona/Private/SAnimNotifyPanel.cpp b/Engine/Source/Editor/Persona/Private/SAnimNotifyPanel.cpp
--- a/Engine/Source/Editor/Persona/Private/SAnimNotifyPanel.cpp
+++ b/Engine/Source/Editor/Persona/Private/SAnimNotifyPanel.cpp
@@ -10,6 +10,8 @@
 #include "AssetRegistryModule.h"
 #include "Editor/UnrealEd/Public/AssetNotifications.h"
 #include "AssetSelection.h"
+#include "STextEntryPopup.h"
+#include "SExpandableArea.h"
 
 // Track Panel drawing
 const float NotificationTrackHeight = 20.0f;
diff --git a/Engine/Source/Editor/Persona/Private/SAnimViewportToolBar.cpp b/Engine/Source/Editor/Persona/Private/SAnimViewportToolBar.cpp
--- a/Engine/Source/Editor/Persona/Private/SAnimViewportToolBar.cpp
+++ b/Engine/Source/Editor/Persona/Private/SAnimViewportToolBar.cpp
@@ -15,6 +15,8 @@
 #include "Editor/UnrealEd/Public/STransformViewportToolbar.h"
 #include "SEditorViewportViewMenu.h"
 #include "SColorPicker.h"
+#include "SNumericEntryBox.h"
+#include "STextComboBox.h"
 
 
 #define LOCTEXT_NAMESPACE "AnimViewportToolBar"
diff --git a/Engine/Source/Editor/Persona/Private/SAnimationBlendSpaceBase.cpp b/Engine/Source/Editor/Persona/Private/SAnimationBlendSpaceBase.cpp
--- a/Engine/Source/Editor/Persona/Private/SAnimationBlendSpaceBase.cpp
+++ b/Engine/Source/Editor/Persona/Private/SAnimationBlendSpaceBase.cpp
@@ -10,6 +10,8 @@
 #include "DragAndDrop/AssetDragDropOp.h"
 #include "ScopedTransaction.h"
 #include "AnimPreviewInstance.h"
+#include "SExpandableArea.h"
+#include "SNotificationList.h"
 
 DEFINE_LOG_CATEGORY_STATIC(LogBlendSpaceBase, Log, All);
 
diff --git a/Engine/Source/Editor/Persona/Private/SAnimationBlendSpaceBase.h b/Engine/Source/Editor/Persona/Private/SAnimationBlendSpaceBase.h
--- a/Engine/Source/Editor/Persona/Private/SAnimationBlendSpaceBase.h
+++ b/Engine/Source/Editor/Persona/Private/SAnimationBlendSpaceBase.h
@@ -2,7 +2,7 @@
 
 #pragma once
 
-#include "Slate.h"
+#include "SlateBasics.h"
 #include "Sorting.h"
 
 extern FText GetErrorMessageForSequence(UBlendSpaceBase* BlendSpace, UAnimSequence* AnimSequence, FVector SampleValue, int32 OriginalIndex);
@@ -14,6 +14,7 @@ DECLARE_DELEGATE(FOnSamplesUpdated);
 // delegate to refresh sample list
 DECLARE_DELEGATE( FOnRefreshSamples )
 
+struct FNotificationInfo;
 // delegate to show user notifications
 DECLARE_DELEGATE_OneParam( FOnNotifyUser, FNotificationInfo& )
 
diff --git a/Engine/Source/Editor/Persona/Private/SRetargetSourceWindow.cpp b/Engine/Source/Editor/Persona/Private/SRetargetSourceWindow.cpp
--- a/Engine/Source/Editor/Persona/Private/SRetargetSourceWindow.cpp
+++ b/Engine/Source/Editor/Persona/Private/SRetargetSourceWindow.cpp
@@ -9,6 +9,10 @@
 #include "Editor/PropertyEditor/Public/PropertyEditorModule.h"
 #include "Editor/ContentBrowser/Public/ContentBrowserModule.h"
 #include "AssetNotifications.h"
+#include "SSearchBox.h"
+#include "SInlineEditableTextBlock.h"
+#include "SNotificationList.h"
+#include "NotificationManager.h"
 
 #define LOCTEXT_NAMESPACE "SRetargetSourceWindow"
 
diff --git a/Engine/Source/Editor/Persona/Private/SRigWindow.cpp b/Engine/Source/Editor/Persona/Private/SRigWindow.cpp
--- a/Engine/Source/Editor/Persona/Private/SRigWindow.cpp
+++ b/Engine/Source/Editor/Persona/Private/SRigWindow.cpp
@@ -12,6 +12,8 @@
 #include "AssetNotifications.h"
 #include "Animation/Rig.h"
 #include "BoneSelectionWidget.h"
+#include "SSearchBox.h"
+#include "SInlineEditableTextBlock.h"
 
 #define LOCTEXT_NAMESPACE "SRigWindow"
 
diff --git a/Engine/Source/Editor/Persona/Private/SSkeletonAnimNotifies.cpp b/Engine/Source/Editor/Persona/Private/SSkeletonAnimNotifies.cpp
--- a/Engine/Source/Editor/Persona/Private/SSkeletonAnimNotifies.cpp
+++ b/Engine/Source/Editor/Persona/Private/SSkeletonAnimNotifies.cpp
@@ -7,6 +7,10 @@
 #include "SSkeletonAnimNotifies.h"
 #include "AssetRegistryModule.h"
 #include "ScopedTransaction.h"
+#include "SSearchBox.h"
+#include "SInlineEditableTextBlock.h"
+#include "SNotificationList.h"
+#include "NotificationManager.h"
 
 #define LOCTEXT_NAMESPACE "SkeletonAnimNotifies"
 
diff --git a/Engine/Source/Editor/Persona/Private/SSkeletonSlotGroupNames.cpp b/Engine/Source/Editor/Persona/Private/SSkeletonSlotGroupNames.cpp
--- a/Engine/Source/Editor/Persona/Private/SSkeletonSlotGroupNames.cpp
+++ b/Engine/Source/Editor/Persona/Private/SSkeletonSlotGroupNames.cpp
@@ -7,6 +7,10 @@
 #include "SSkeletonSlotGroupNames.h"
 #include "AssetRegistryModule.h"
 #include "ScopedTransaction.h"
+#include "SSearchBox.h"
+#include "SInlineEditableTextBlock.h"
+#include "NotificationManager.h"
+#include "SNotificationList.h"
 
 #define LOCTEXT_NAMESPACE "SkeletonSlotGroupNames"
 
diff --git a/Engine/Source/Editor/Persona/Private/SSkeletonSmartNameManager.cpp b/Engine/Source/Editor/Persona/Private/SSkeletonSmartNameManager.cpp
--- a/Engine/Source/Editor/Persona/Private/SSkeletonSmartNameManager.cpp
+++ b/Engine/Source/Editor/Persona/Private/SSkeletonSmartNameManager.cpp
@@ -7,6 +7,9 @@
 #include "ScopedTransaction.h"
 #include "AssetRegistryModule.h"
 #include "AssetData.h"
+#include "SSearchBox.h"
+#include "SInlineEditableTextBlock.h"
+#include "STextEntryPopup.h"
 
 #define LOCTEXT_NAMESPACE "SkeletonSmartNameManager"
 
diff --git a/Engine/Source/Editor/Persona/Private/SSkeletonTree.cpp b/Engine/Source/Editor/Persona/Private/SSkeletonTree.cpp
--- a/Engine/Source/Editor/Persona/Private/SSkeletonTree.cpp
+++ b/Engine/Source/Editor/Persona/Private/SSkeletonTree.cpp
@@ -28,6 +28,11 @@
 #include "Factories.h"
 #include "Developer/MeshUtilities/Public/MeshUtilities.h"
 #include "UnrealExporter.h"
+#include "SSearchBox.h"
+#include "SInlineEditableTextBlock.h"
+#include "SNotificationList.h"
+#include "NotificationManager.h"
+#include "GenericCommands.h"
 
 #define LOCTEXT_NAMESPACE "SSkeletonTree"
 
diff --git a/Engine/Source/Editor/PhAT/Private/SPhATNewAssetDlg.cpp b/Engine/Source/Editor/PhAT/Private/SPhATNewAssetDlg.cpp
--- a/Engine/Source/Editor/PhAT/Private/SPhATNewAssetDlg.cpp
+++ b/Engine/Source/Editor/PhAT/Private/SPhATNewAssetDlg.cpp
@@ -2,6 +2,8 @@
 
 #include "PhATModule.h"
 #include "SPhATNewAssetDlg.h"
+#include "SNumericEntryBox.h"
+#include "STextComboBox.h"
 
 // Add in the constants from the static mesh editor as we need them here too
 const int32 DefaultHullCount = 4;
diff --git a/Engine/Source/Editor/PlacementMode/Private/SPlacementModeTools.cpp b/Engine/Source/Editor/PlacementMode/Private/SPlacementModeTools.cpp
--- a/Engine/Source/Editor/PlacementMode/Private/SPlacementModeTools.cpp
+++ b/Engine/Source/Editor/PlacementMode/Private/SPlacementModeTools.cpp
@@ -9,6 +9,8 @@
 #include "IBspModeModule.h"
 #include "KismetEditorUtilities.h"
 #include "EditorClassUtils.h"
+#include "SSearchBox.h"
+#include "SWidgetSwitcher.h"
 
 /**
  * These are the tab indexes, if the tabs are reorganized you need to adjust the
diff --git a/Engine/Source/Editor/PropertyEditor/Private/UserInterface/PropertyEditor/SPropertyEditorNumeric.h b/Engine/Source/Editor/PropertyEditor/Private/UserInterface/PropertyEditor/SPropertyEditorNumeric.h
--- a/Engine/Source/Editor/PropertyEditor/Private/UserInterface/PropertyEditor/SPropertyEditorNumeric.h
+++ b/Engine/Source/Editor/PropertyEditor/Private/UserInterface/PropertyEditor/SPropertyEditorNumeric.h
@@ -2,6 +2,7 @@
 #pragma once
 
 #include "PropertyEditorConstants.h"
+#include "SNumericEntryBox.h"
 
 #define LOCTEXT_NAMESPACE "PropertyEditor"
 
@@ -39,40 +40,40 @@ public:
 		NumericType ClampMin = TNumericLimits<NumericType>::Lowest();
 		NumericType ClampMax = TNumericLimits<NumericType>::Max();
 
-		if( !ClampMinString.IsEmpty() )
+		if (!ClampMinString.IsEmpty())
 		{
-			 TTypeFromString<NumericType>::FromString(ClampMin, *ClampMinString );
+			TTypeFromString<NumericType>::FromString(ClampMin, *ClampMinString);
 		}
 
-		if( !ClampMaxString.IsEmpty() )
+		if (!ClampMaxString.IsEmpty())
 		{
-			TTypeFromString<NumericType>::FromString(ClampMax, *ClampMaxString );
+			TTypeFromString<NumericType>::FromString(ClampMax, *ClampMaxString);
 		}
 
 		NumericType UIMin = TNumericLimits<NumericType>::Lowest();
 		NumericType UIMax = TNumericLimits<NumericType>::Max();
-		TTypeFromString<NumericType>::FromString( UIMin, *UIMinString );
-		TTypeFromString<NumericType>::FromString( UIMax, *UIMaxString );
+		TTypeFromString<NumericType>::FromString(UIMin, *UIMinString);
+		TTypeFromString<NumericType>::FromString(UIMax, *UIMaxString);
 
 		NumericType SliderExponent = NumericType(1);
 		if (SliderExponentString.Len())
 		{
-			TTypeFromString<NumericType>::FromString( SliderExponent, *SliderExponentString );
+			TTypeFromString<NumericType>::FromString(SliderExponent, *SliderExponentString);
 		}
 
 		NumericType Delta = NumericType(0);
-		if ( DeltaString.Len() )
+		if (DeltaString.Len())
 		{
-			TTypeFromString<NumericType>::FromString( Delta, *DeltaString );
+			TTypeFromString<NumericType>::FromString(Delta, *DeltaString);
 		}
-		
+
 		if (ClampMin >= ClampMax && (ClampMinString.Len() || ClampMaxString.Len()))
 		{
-			UE_LOG(LogPropertyNode, Warning, TEXT("Clamp Min (%s) >= Clamp Max (%s) for Ranged Numeric"), *ClampMinString, *ClampMaxString );
+			UE_LOG(LogPropertyNode, Warning, TEXT("Clamp Min (%s) >= Clamp Max (%s) for Ranged Numeric"), *ClampMinString, *ClampMaxString);
 		}
 
-		const NumericType ActualUIMin = FMath::Max(UIMin,ClampMin);
-		const NumericType ActualUIMax = FMath::Min(UIMax,ClampMax);
+		const NumericType ActualUIMin = FMath::Max(UIMin, ClampMin);
+		const NumericType ActualUIMax = FMath::Min(UIMax, ClampMax);
 
 		TOptional<NumericType> MinValue = ClampMinString.Len() ? ClampMin : TOptional<NumericType>();
 		TOptional<NumericType> MaxValue = ClampMaxString.Len() ? ClampMax : TOptional<NumericType>();
@@ -81,7 +82,7 @@ public:
 
 		if (ActualUIMin >= ActualUIMax && (MetaUIMinString.Len() || MetaUIMaxString.Len()))
 		{
-			UE_LOG(LogPropertyNode, Warning, TEXT("UI Min (%s) >= UI Max (%s) for Ranged Numeric"), *UIMinString, *UIMaxString );
+			UE_LOG(LogPropertyNode, Warning, TEXT("UI Min (%s) >= UI Max (%s) for Ranged Numeric"), *UIMinString, *UIMaxString);
 		}
 
 		FObjectPropertyNode* ObjectPropertyNode = PropertyNode->FindObjectItemParent();
@@ -90,25 +91,25 @@ public:
 
 		ChildSlot
 			[
-				SAssignNew( PrimaryWidget, SNumericEntryBox<NumericType> )
+				SAssignNew(PrimaryWidget, SNumericEntryBox<NumericType>)
 				// Only allow spinning if we have a single value
 				.AllowSpin(bAllowSpin)
-				.Value( this, &SPropertyEditorNumeric<NumericType>::OnGetValue )
-				.Font( InArgs._Font )
+				.Value(this, &SPropertyEditorNumeric<NumericType>::OnGetValue)
+				.Font(InArgs._Font)
 				.MinValue(MinValue)
 				.MaxValue(MaxValue)
 				.MinSliderValue(SliderMinValue)
 				.MaxSliderValue(SliderMaxValue)
 				.SliderExponent(SliderExponent)
 				.Delta(Delta)
-				.UndeterminedString( LOCTEXT("MultipleValues", "Multiple Values").ToString() )
-				.OnValueChanged( this, &SPropertyEditorNumeric<NumericType>::OnValueChanged )
-				.OnValueCommitted( this, &SPropertyEditorNumeric<NumericType>::OnValueCommitted )
-				.OnBeginSliderMovement( this, &SPropertyEditorNumeric<NumericType>::OnBeginSliderMovement )
-				.OnEndSliderMovement( this, &SPropertyEditorNumeric<NumericType>::OnEndSliderMovement )
+				.UndeterminedString(LOCTEXT("MultipleValues", "Multiple Values").ToString())
+				.OnValueChanged(this, &SPropertyEditorNumeric<NumericType>::OnValueChanged)
+				.OnValueCommitted(this, &SPropertyEditorNumeric<NumericType>::OnValueCommitted)
+				.OnBeginSliderMovement(this, &SPropertyEditorNumeric<NumericType>::OnBeginSliderMovement)
+				.OnEndSliderMovement(this, &SPropertyEditorNumeric<NumericType>::OnEndSliderMovement)
 			];
 
-		SetEnabled( TAttribute<bool>( this, &SPropertyEditorNumeric<NumericType>::CanEdit ) );
+		SetEnabled(TAttribute<bool>(this, &SPropertyEditorNumeric<NumericType>::CanEdit));
 	}
 
 	virtual bool SupportsKeyboardFocus() const override
diff --git a/Engine/Source/Editor/Sequencer/Private/Sequencer.cpp b/Engine/Source/Editor/Sequencer/Private/Sequencer.cpp
--- a/Engine/Source/Editor/Sequencer/Private/Sequencer.cpp
+++ b/Engine/Source/Editor/Sequencer/Private/Sequencer.cpp
@@ -32,6 +32,8 @@
 #include "MovieSceneInstance.h"
 #include "IKeyArea.h"
 #include "SnappingUtils.h"
+#include "STextEntryPopup.h"
+#include "GenericCommands.h"
 
 #define LOCTEXT_NAMESPACE "Sequencer"
 
diff --git a/Engine/Source/Editor/StaticMeshEditor/Private/StaticMeshEditor.cpp b/Engine/Source/Editor/StaticMeshEditor/Private/StaticMeshEditor.cpp
--- a/Engine/Source/Editor/StaticMeshEditor/Private/StaticMeshEditor.cpp
+++ b/Engine/Source/Editor/StaticMeshEditor/Private/StaticMeshEditor.cpp
@@ -24,6 +24,11 @@
 
 #include "Runtime/Analytics/Analytics/Public/Interfaces/IAnalyticsProvider.h"
 #include "EngineAnalytics.h"
+#include "SDockTab.h"
+#include "GenericCommands.h"
+#include "STextComboBox.h"
+#include "SNotificationList.h"
+#include "NotificationManager.h"
 
 #define LOCTEXT_NAMESPACE "StaticMeshEditor"
 
diff --git a/Engine/Source/Editor/StatsViewer/Private/SStatsViewer.cpp b/Engine/Source/Editor/StatsViewer/Private/SStatsViewer.cpp
--- a/Engine/Source/Editor/StatsViewer/Private/SStatsViewer.cpp
+++ b/Engine/Source/Editor/StatsViewer/Private/SStatsViewer.cpp
@@ -15,6 +15,9 @@
 #include "ActorArrayHyperlinkColumn.h"
 #include "StatsViewerUtils.h"
 #include "EngineVersion.h"
+#include "SSearchBox.h"
+#include "SNotificationList.h"
+#include "NotificationManager.h"
 
 DEFINE_LOG_CATEGORY_STATIC(LogStatsViewer, Log, All);
 
diff --git a/Engine/Source/Editor/StatsViewer/Private/StatsPages/StaticMeshLightingInfoStatsPage.cpp b/Engine/Source/Editor/StatsViewer/Private/StatsPages/StaticMeshLightingInfoStatsPage.cpp
--- a/Engine/Source/Editor/StatsViewer/Private/StatsPages/StaticMeshLightingInfoStatsPage.cpp
+++ b/Engine/Source/Editor/StatsViewer/Private/StatsPages/StaticMeshLightingInfoStatsPage.cpp
@@ -7,6 +7,7 @@
 #include "Editor/PropertyEditor/Public/IPropertyTableRow.h"
 #include "BusyCursor.h"
 #include "ScopedTransaction.h"
+#include "STextEntryPopup.h"
 
 #define LOCTEXT_NAMESPACE "Editor.StatsViewer.StaticMeshLightingInfo"
 
diff --git a/Engine/Source/Editor/TextureEditor/Private/TextureEditorToolkit.cpp b/Engine/Source/Editor/TextureEditor/Private/TextureEditorToolkit.cpp
--- a/Engine/Source/Editor/TextureEditor/Private/TextureEditorToolkit.cpp
+++ b/Engine/Source/Editor/TextureEditor/Private/TextureEditorToolkit.cpp
@@ -2,6 +2,8 @@
 
 #include "TextureEditorPrivatePCH.h"
 #include "Factories.h"
+#include "SDockTab.h"
+#include "SNumericEntryBox.h"
 
 #define LOCTEXT_NAMESPACE "FTextureEditorToolkit"
 
diff --git a/Engine/Source/Editor/UMGEditor/Private/Designer/SDesignerView.cpp b/Engine/Source/Editor/UMGEditor/Private/Designer/SDesignerView.cpp
--- a/Engine/Source/Editor/UMGEditor/Private/Designer/SDesignerView.cpp
+++ b/Engine/Source/Editor/UMGEditor/Private/Designer/SDesignerView.cpp
@@ -19,6 +19,8 @@
 #include "DesignerCommands.h"
 #include "STransformHandle.h"
 #include "Runtime/Engine/Classes/Engine/RendererSettings.h"
+#include "SDPIScaler.h"
+#include "SNumericEntryBox.h"
 
 #define LOCTEXT_NAMESPACE "UMG"
 
diff --git a/Engine/Source/Editor/UMGEditor/Private/TabFactory/AnimationTabSummoner.cpp b/Engine/Source/Editor/UMGEditor/Private/TabFactory/AnimationTabSummoner.cpp
--- a/Engine/Source/Editor/UMGEditor/Private/TabFactory/AnimationTabSummoner.cpp
+++ b/Engine/Source/Editor/UMGEditor/Private/TabFactory/AnimationTabSummoner.cpp
@@ -8,6 +8,9 @@
 #include "WidgetAnimation.h"
 #include "MovieScene.h"
 #include "UMGStyle.h"
+#include "SSearchBox.h"
+#include "SInlineEditableTextBlock.h"
+#include "GenericCommands.h"
 
 #define LOCTEXT_NAMESPACE "UMG"
 
diff --git a/Engine/Source/Editor/UnrealEd/Private/AnimSetViewer/LODSettingsWindow.cpp b/Engine/Source/Editor/UnrealEd/Private/AnimSetViewer/LODSettingsWindow.cpp
--- a/Engine/Source/Editor/UnrealEd/Private/AnimSetViewer/LODSettingsWindow.cpp
+++ b/Engine/Source/Editor/UnrealEd/Private/AnimSetViewer/LODSettingsWindow.cpp
@@ -1,12 +1,14 @@
 // Copyright 1998-2014 Epic Games, Inc. All Rights Reserved.
 
 #include "UnrealEd.h"
-#include "Slate.h"
+#include "SlateBasics.h"
 #include "LODSettingsWindow.h"
 #include "LODUtilities.h"
 #include "MeshUtilities.h"
 #include "FbxMeshUtils.h"
 #include "SAssetSearchBox.h"
+#include "SExpandableArea.h"
+#include "STextComboBox.h"
 
 #define LOCTEXT_NAMESPACE "SkeletalMeshSimplification"
 
diff --git a/Engine/Source/Editor/UnrealEd/Private/Dialogs/Dialogs.cpp b/Engine/Source/Editor/UnrealEd/Private/Dialogs/Dialogs.cpp
--- a/Engine/Source/Editor/UnrealEd/Private/Dialogs/Dialogs.cpp
+++ b/Engine/Source/Editor/UnrealEd/Private/Dialogs/Dialogs.cpp
@@ -4,10 +4,11 @@
 #include "UnrealEd.h"
 #include "TexAlignTools.h"
 #include "ScopedTransaction.h"
-#include "Slate.h"
+#include "SlateBasics.h"
 #include "ObjectTools.h"
 #include "Editor/MainFrame/Public/MainFrame.h"
 #include "DesktopPlatformModule.h"
+#include "SHyperlink.h"
 DEFINE_LOG_CATEGORY_STATIC(LogDialogs, Log, All);
 
 #define LOCTEXT_NAMESPACE "Dialogs"
diff --git a/Engine/Source/Editor/UnrealEd/Private/GraphEditor.cpp b/Engine/Source/Editor/UnrealEd/Private/GraphEditor.cpp
--- a/Engine/Source/Editor/UnrealEd/Private/GraphEditor.cpp
+++ b/Engine/Source/Editor/UnrealEd/Private/GraphEditor.cpp
@@ -5,6 +5,7 @@
 #include "GraphEditor.h"
 #include "Editor/GraphEditor/Public/GraphEditorModule.h"
 #include "BlueprintUtilities.h"
+#include "SMissingWidget.h"
 
 // List of all active GraphEditor wrappers
 TArray< TWeakPtr<SGraphEditor> > SGraphEditor::AllInstances;
diff --git a/Engine/Source/Editor/UnrealEd/Private/PlayLevel.cpp b/Engine/Source/Editor/UnrealEd/Private/PlayLevel.cpp
--- a/Engine/Source/Editor/UnrealEd/Private/PlayLevel.cpp
+++ b/Engine/Source/Editor/UnrealEd/Private/PlayLevel.cpp
@@ -34,6 +34,9 @@
 #include "Runtime/Engine/Classes/Engine/RendererSettings.h"
 #include "SScissorRectBox.h"
 #include "Online.h"
+#include "SNotificationList.h"
+#include "SDPIScaler.h"
+#include "NotificationManager.h"
 
 DEFINE_LOG_CATEGORY_STATIC(LogPlayLevel, Log, All);
 
diff --git a/Engine/Source/Editor/UnrealEd/Private/SComponentClassCombo.cpp b/Engine/Source/Editor/UnrealEd/Private/SComponentClassCombo.cpp
--- a/Engine/Source/Editor/UnrealEd/Private/SComponentClassCombo.cpp
+++ b/Engine/Source/Editor/UnrealEd/Private/SComponentClassCombo.cpp
@@ -1,7 +1,7 @@
 // Copyright 1998-2014 Epic Games, Inc. All Rights Reserved.
 
 #include "UnrealEd.h"
-#include "Slate.h"
+#include "SlateBasics.h"
 #include "SComponentClassCombo.h"
 #include "ComponentAssetBroker.h"
 #include "ClassIconFinder.h"
@@ -9,6 +9,7 @@
 #include "IDocumentation.h"
 #include "SListViewSelectorDropdownMenu.h"
 #include "EditorClassUtils.h"
+#include "SSearchBox.h"
 
 #define LOCTEXT_NAMESPACE "ComponentClassCombo"
 
diff --git a/Engine/Source/Editor/UnrealEd/Private/SCurveEditor.cpp b/Engine/Source/Editor/UnrealEd/Private/SCurveEditor.cpp
--- a/Engine/Source/Editor/UnrealEd/Private/SCurveEditor.cpp
+++ b/Engine/Source/Editor/UnrealEd/Private/SCurveEditor.cpp
@@ -6,6 +6,8 @@
 #include "SCurveEditor.h"
 #include "ScopedTransaction.h"
 #include "SColorGradientEditor.h"
+#include "GenericCommands.h"
+#include "SNumericEntryBox.h"
 
 #define LOCTEXT_NAMESPACE "SCurveEditor"
 
diff --git a/Engine/Source/Editor/UnrealEd/Private/SSocketManager.cpp b/Engine/Source/Editor/UnrealEd/Private/SSocketManager.cpp
--- a/Engine/Source/Editor/UnrealEd/Private/SSocketManager.cpp
+++ b/Engine/Source/Editor/UnrealEd/Private/SSocketManager.cpp
@@ -12,6 +12,8 @@
 
 #include "Runtime/Analytics/Analytics/Public/Interfaces/IAnalyticsProvider.h"
 #include "EngineAnalytics.h"
+#include "SInlineEditableTextBlock.h"
+#include "GenericCommands.h"
 
 #define LOCTEXT_NAMESPACE "SSCSSocketManagerEditor"
 
diff --git a/Engine/Source/Editor/UnrealEd/Private/Toolkits/BaseToolkit.cpp b/Engine/Source/Editor/UnrealEd/Private/Toolkits/BaseToolkit.cpp
--- a/Engine/Source/Editor/UnrealEd/Private/Toolkits/BaseToolkit.cpp
+++ b/Engine/Source/Editor/UnrealEd/Private/Toolkits/BaseToolkit.cpp
@@ -4,6 +4,9 @@
 #include "Toolkits/BaseToolkit.h"
 #include "Toolkits/IToolkitHost.h"
 #include "Toolkits/ToolkitManager.h"
+#include "SDockableTab.h"
+#include "SDockTabStack.h"
+#include "SDockTab.h"
 
 #define LOCTEXT_NAMESPACE "BaseToolkit"
 
diff --git a/Engine/Source/Editor/UnrealEd/Public/GraphEditor.h b/Engine/Source/Editor/UnrealEd/Public/GraphEditor.h
--- a/Engine/Source/Editor/UnrealEd/Public/GraphEditor.h
+++ b/Engine/Source/Editor/UnrealEd/Public/GraphEditor.h
@@ -3,10 +3,11 @@
 #pragma once
 
 #include "UnrealEd.h"
-#include "Slate.h"
+#include "SlateBasics.h"
 #include "BlueprintUtilities.h"
 
 class UEdGraph;
+struct FNotificationInfo;
 
 DECLARE_DELEGATE_ThreeParams( FOnNodeTextCommitted, const FText&, ETextCommit::Type, UEdGraphNode* );
 DECLARE_DELEGATE_RetVal_TwoParams( bool, FOnNodeVerifyTextCommit, const FText&, UEdGraphNode* );
diff --git a/Engine/Source/Editor/UserFeedback/Private/UserFeedback.cpp b/Engine/Source/Editor/UserFeedback/Private/UserFeedback.cpp
--- a/Engine/Source/Editor/UserFeedback/Private/UserFeedback.cpp
+++ b/Engine/Source/Editor/UserFeedback/Private/UserFeedback.cpp
@@ -1,13 +1,15 @@
 // Copyright 1998-2014 Epic Games, Inc. All Rights Reserved.
 
 #include "IUserFeedbackModule.h"
-#include "Slate.h"
+#include "SlateBasics.h"
 #include "SlateStyle.h"
 #include "ModuleManager.h"
 
 #include "IAnalyticsProvider.h"
 #include "EngineAnalytics.h"
 #include "UnrealEd.h"
+#include "SNotificationList.h"
+#include "NotificationManager.h"
 
 #define LOCTEXT_NAMESPACE "UserFeedback"
 
diff --git a/Engine/Source/Programs/CrashReporter/CrashReportClient/Private/CrashDescription.h b/Engine/Source/Programs/CrashReporter/CrashReportClient/Private/CrashDescription.h
--- a/Engine/Source/Programs/CrashReporter/CrashReportClient/Private/CrashDescription.h
+++ b/Engine/Source/Programs/CrashReporter/CrashReportClient/Private/CrashDescription.h
@@ -1,6 +1,7 @@
 // Copyright 1998-2014 Epic Games, Inc. All Rights Reserved.
 
 #pragma once
+#include "GenericPlatformContext.h"
 
 /**
  *	Describes a unified crash, should be used by all platforms.
diff --git a/Engine/Source/Programs/CrashReporter/CrashReportClient/Private/CrashReportClientApp.cpp b/Engine/Source/Programs/CrashReporter/CrashReportClient/Private/CrashReportClientApp.cpp
--- a/Engine/Source/Programs/CrashReporter/CrashReportClient/Private/CrashReportClientApp.cpp
+++ b/Engine/Source/Programs/CrashReporter/CrashReportClient/Private/CrashReportClientApp.cpp
@@ -91,6 +91,7 @@ FPlatformErrorReport LoadErrorReport()
 	FString XMLWerFilename;
 	ErrorReport.FindFirstReportFileWithExtension( XMLWerFilename, TEXT( ".xml" ) );
 
+	extern FCrashDescription& GetCrashDescription();
 	GetCrashDescription() = FCrashDescription( ReportDirectoryAbsolutePath / XMLWerFilename );
 
 #if CRASH_REPORT_UNATTENDED_ONLY
@@ -129,6 +130,7 @@ void RunCrashReportClient(const TCHAR* CommandLine)
 	if( ErrorReport.HasFilesToUpload() )
 	{
 		// Send analytics.
+		extern FCrashDescription& GetCrashDescription();
 		GetCrashDescription().SendAnalytics();
 	}
 
diff --git a/Engine/Source/Programs/CrashReporter/CrashReportClient/Private/GenericErrorReport.cpp b/Engine/Source/Programs/CrashReporter/CrashReportClient/Private/GenericErrorReport.cpp
--- a/Engine/Source/Programs/CrashReporter/CrashReportClient/Private/GenericErrorReport.cpp
+++ b/Engine/Source/Programs/CrashReporter/CrashReportClient/Private/GenericErrorReport.cpp
@@ -66,6 +66,7 @@ bool FGenericErrorReport::SetUserComment(const FText& UserComment)
 
 	FString MachineIDandUserID;
 	// Set global user name ID: will be added to the report
+	extern FCrashDescription& GetCrashDescription();
 	if( FRocketSupport::IsRocket() )
 	{
 		MachineIDandUserID = FString::Printf( TEXT( "!MachineId:%s!EpicAccountId:%s" ), *GetCrashDescription().MachineId, *GetCrashDescription().EpicAccountId );
diff --git a/Engine/Source/Programs/UnrealSync/Private/GUI.cpp b/Engine/Source/Programs/UnrealSync/Private/GUI.cpp
--- a/Engine/Source/Programs/UnrealSync/Private/GUI.cpp
+++ b/Engine/Source/Programs/UnrealSync/Private/GUI.cpp
@@ -2,13 +2,16 @@
 
 #include "UnrealSync.h"
 
-#include "Slate.h"
+#include "SlateBasics.h"
 #include "StandaloneRenderer.h"
 #include "LaunchEngineLoop.h"
 
 #include "P4Env.h"
 
 #include "ProcessHelper.h"
+#include "SWidgetSwitcher.h"
+#include "SDockTab.h"
+#include "SlateFwd.h"
 
 /**
  * Simple text combo box widget.
diff --git a/Engine/Source/Runtime/AppFramework/Private/Framework/Testing/SWidgetGallery.cpp b/Engine/Source/Runtime/AppFramework/Private/Framework/Testing/SWidgetGallery.cpp
--- a/Engine/Source/Runtime/AppFramework/Private/Framework/Testing/SWidgetGallery.cpp
+++ b/Engine/Source/Runtime/AppFramework/Private/Framework/Testing/SWidgetGallery.cpp
@@ -1,10 +1,18 @@
 // Copyright 1998-2014 Epic Games, Inc. All Rights Reserved.
 
 #include "AppFrameworkPrivatePCH.h"
+#include "SSearchBox.h"
+#include "SDockTab.h"
 #include "SWidgetGallery.h"
 #include "TestStyle.h"
 
 #include "SUserWidgetTest.h"
+#include "SBreadcrumbTrail.h"
+#include "SThrobber.h"
+#include "SHyperlink.h"
+#include "STextComboBox.h"
+#include "SVolumeControl.h"
+#include "STextComboPopup.h"
 
 
 #define LOCTEXT_NAMESPACE "WidgetGallery"
@@ -485,6 +493,7 @@ public:
 								.Padding(0.0f, 4.0f, 0.0f, 0.0f)
 								[
 									SNew(SSearchBox)
+										.DelayChangeNotificationsWhileTyping(true)
 										.OnSearch(this, &SWidgetGallery::HandleSearchBoxSearch)
 								]			
 						]
diff --git a/Engine/Source/Runtime/AppFramework/Private/Widgets/Testing/STestSuite.cpp b/Engine/Source/Runtime/AppFramework/Private/Widgets/Testing/STestSuite.cpp
--- a/Engine/Source/Runtime/AppFramework/Private/Widgets/Testing/STestSuite.cpp
+++ b/Engine/Source/Runtime/AppFramework/Private/Widgets/Testing/STestSuite.cpp
@@ -13,6 +13,14 @@
 #include "SlateRenderTransform.h"
 #include "SlateLayoutTransform.h"
 #include "SWebBrowser.h"
+#include "SInlineEditableTextBlock.h"
+#include "STextEntryPopup.h"
+#include "SDPIScaler.h"
+#include "SNotificationList.h"
+#include "SNumericEntryBox.h"
+#include "SDockTab.h"
+#include "SSearchBox.h"
+#include "SVolumeControl.h"
 
 #define LOCTEXT_NAMESPACE "STestSuite"
 
diff --git a/Engine/Source/Runtime/AppFramework/Private/Widgets/Workflow/SWizard.cpp b/Engine/Source/Runtime/AppFramework/Private/Widgets/Workflow/SWizard.cpp
--- a/Engine/Source/Runtime/AppFramework/Private/Widgets/Workflow/SWizard.cpp
+++ b/Engine/Source/Runtime/AppFramework/Private/Widgets/Workflow/SWizard.cpp
@@ -2,6 +2,7 @@
 
 #include "AppFrameworkPrivatePCH.h"
 #include "SWizard.h"
+#include "SWidgetSwitcher.h"
 
 
 #define LOCTEXT_NAMESPACE "SWizard"
@@ -348,5 +349,15 @@ EVisibility SWizard::HandlePrevButtonVisibility( ) const
 	return EVisibility::Hidden;
 }
 
+int32 SWizard::GetNumPages() const
+{
+	return WidgetSwitcher->GetNumWidgets();
+}
+
+int32 SWizard::GetPageIndex(const TSharedRef<SWidget>& PageWidget) const
+{
+	return WidgetSwitcher->GetWidgetIndex(PageWidget);
+}
+
 
 #undef LOCTEXT_NAMESPACE
diff --git a/Engine/Source/Runtime/AppFramework/Public/Widgets/Workflow/SWizard.h b/Engine/Source/Runtime/AppFramework/Public/Widgets/Workflow/SWizard.h
--- a/Engine/Source/Runtime/AppFramework/Public/Widgets/Workflow/SWizard.h
+++ b/Engine/Source/Runtime/AppFramework/Public/Widgets/Workflow/SWizard.h
@@ -208,10 +208,7 @@ public:
 	 *
 	 * @return The number of pages.
 	 */
-	int32 GetNumPages( ) const
-	{
-		return WidgetSwitcher->GetNumWidgets();
-	}
+	int32 GetNumPages( ) const;
 
 	/**
 	 * Gets the index of the specified wizard page widget.
@@ -219,10 +216,7 @@ public:
 	 * @param PageWidget The page widget to get the index for.
 	 * @return The index of the page, or INDEX_NONE if not found.
 	 */
-	int32 GetPageIndex( const TSharedRef<SWidget>& PageWidget ) const
-	{
-		return WidgetSwitcher->GetWidgetIndex(PageWidget);
-	}
+	int32 GetPageIndex( const TSharedRef<SWidget>& PageWidget ) const;
 
 	/**
 	 * Attempts to show the page with the specified index.
diff --git a/Engine/Source/Runtime/Engine/Private/GameEngine.cpp b/Engine/Source/Runtime/Engine/Private/GameEngine.cpp
--- a/Engine/Source/Runtime/Engine/Private/GameEngine.cpp
+++ b/Engine/Source/Runtime/Engine/Private/GameEngine.cpp
@@ -18,7 +18,7 @@
 #include "Engine/RendererSettings.h"
 #include "AVIWriter.h"
 
-#include "Slate.h"
+#include "SlateBasics.h"
 #include "Slate/SceneViewport.h"
 #include "SVirtualJoystick.h"
 
@@ -29,6 +29,7 @@
 #include "IHeadMountedDisplay.h"
 #include "RendererInterface.h"
 #include "HotReloadInterface.h"
+#include "SDPIScaler.h"
 
 DEFINE_LOG_CATEGORY_STATIC(LogEngine, Log, All);
 
diff --git a/Engine/Source/Runtime/Engine/Private/GameFramework/TouchInterface.cpp b/Engine/Source/Runtime/Engine/Private/GameFramework/TouchInterface.cpp
--- a/Engine/Source/Runtime/Engine/Private/GameFramework/TouchInterface.cpp
+++ b/Engine/Source/Runtime/Engine/Private/GameFramework/TouchInterface.cpp
@@ -2,7 +2,8 @@
 
 #include "EnginePrivate.h"
 #include "GameFramework/TouchInterface.h"
-#include "Slate.h"
+#include "SlateBasics.h"
+#include "SVirtualJoystick.h"
 
 UTouchInterface::UTouchInterface(const FObjectInitializer& ObjectInitializer)
 	: Super(ObjectInitializer)
diff --git a/Engine/Source/Runtime/Engine/Private/Materials/Material.cpp b/Engine/Source/Runtime/Engine/Private/Materials/Material.cpp
--- a/Engine/Source/Runtime/Engine/Private/Materials/Material.cpp
+++ b/Engine/Source/Runtime/Engine/Private/Materials/Material.cpp
@@ -36,7 +36,9 @@
 
 #if WITH_EDITOR
 #include "UnrealEd.h"
-#include "Slate.h"  // For AddNotification
+#include "SlateBasics.h"  // For AddNotification
+#include "SNotificationList.h"
+#include "NotificationManager.h"
 #endif
 
 #define LOCTEXT_NAMESPACE "Material"
diff --git a/Engine/Source/Runtime/Engine/Private/PlayerController.cpp b/Engine/Source/Runtime/Engine/Private/PlayerController.cpp
--- a/Engine/Source/Runtime/Engine/Private/PlayerController.cpp
+++ b/Engine/Source/Runtime/Engine/Private/PlayerController.cpp
@@ -14,7 +14,7 @@
 #include "OnlineSubsystemUtils.h"
 #include "IHeadMountedDisplay.h"
 #include "IForceFeedbackSystem.h"
-#include "Slate.h"
+#include "SlateBasics.h"
 #include "GameFramework/TouchInterface.h"
 #include "DisplayDebugHelpers.h"
 #include "Matinee/InterpTrackInstDirector.h"
@@ -27,6 +27,7 @@
 #include "Camera/CameraActor.h"
 #include "GenericPlatform/IForceFeedbackSystem.h"
 #include "Engine/InputDelegateBinding.h"
+#include "SVirtualJoystick.h"
 
 DEFINE_LOG_CATEGORY(LogPlayerController);
 
diff --git a/Engine/Source/Runtime/Engine/Private/UnrealClient.cpp b/Engine/Source/Runtime/Engine/Private/UnrealClient.cpp
--- a/Engine/Source/Runtime/Engine/Private/UnrealClient.cpp
+++ b/Engine/Source/Runtime/Engine/Private/UnrealClient.cpp
@@ -10,13 +10,14 @@
 #include "../../Renderer/Private/ScenePrivate.h"
 #include "HighResScreenshot.h"
 
-#include "Slate.h"
+#include "SlateBasics.h"
 #include "SNotificationList.h"
 #include "Engine/PostProcessVolume.h"
 #include "EngineModule.h"
 #include "EngineModule.h"
 #include "ContentStreaming.h"
 #include "SceneUtils.h"
+#include "NotificationManager.h"
 
 DEFINE_LOG_CATEGORY_STATIC(LogClient, Log, All);
 
diff --git a/Engine/Source/Runtime/Engine/Private/UnrealEngine.cpp b/Engine/Source/Runtime/Engine/Private/UnrealEngine.cpp
--- a/Engine/Source/Runtime/Engine/Private/UnrealEngine.cpp
+++ b/Engine/Source/Runtime/Engine/Private/UnrealEngine.cpp
@@ -12,7 +12,7 @@
 #include "FileManagerGeneric.h"
 #include "Database.h"
 #include "SkeletalMeshMerge.h"
-#include "Slate.h"
+#include "SlateBasics.h"
 #include "RenderCore.h"
 #include "ShaderCompiler.h"
 #include "ColorList.h"
@@ -85,6 +85,7 @@
 #include "HotReloadInterface.h"
 #include "STestSuite.h"
 #include "Engine/DemoNetDriver.h"
+#include "SThrobber.h"
 
 DEFINE_LOG_CATEGORY_STATIC(LogEngine, Log, All);
 
diff --git a/Engine/Source/Runtime/GameMenuBuilder/Private/SGameMenuPageWidget.cpp b/Engine/Source/Runtime/GameMenuBuilder/Private/SGameMenuPageWidget.cpp
--- a/Engine/Source/Runtime/GameMenuBuilder/Private/SGameMenuPageWidget.cpp
+++ b/Engine/Source/Runtime/GameMenuBuilder/Private/SGameMenuPageWidget.cpp
@@ -2,6 +2,7 @@
 
 #include "GameMenuBuilderPrivatePCH.h"
 #include "Engine/Console.h"
+#include "SDPIScaler.h"
 
 
 FMenuPanel::FMenuPanel()
diff --git a/Engine/Source/Runtime/MoviePlayer/Private/MoviePlayer.cpp b/Engine/Source/Runtime/MoviePlayer/Private/MoviePlayer.cpp
--- a/Engine/Source/Runtime/MoviePlayer/Private/MoviePlayer.cpp
+++ b/Engine/Source/Runtime/MoviePlayer/Private/MoviePlayer.cpp
@@ -3,10 +3,11 @@
 #include "MoviePlayer.h"
 
 #include "Engine.h"
-#include "Slate.h"
+#include "SlateBasics.h"
 #include "SpinLock.h"
 #include "NullMoviePlayer.h"
 #include "DefaultGameMoviePlayer.h"
+#include "SThrobber.h"
 
 IMPLEMENT_MODULE(FDefaultModuleImpl, MoviePlayer);
 
diff --git a/Engine/Source/Runtime/Online/BuildPatchServices/Private/BuildPatchChunk.cpp b/Engine/Source/Runtime/Online/BuildPatchServices/Private/BuildPatchChunk.cpp
--- a/Engine/Source/Runtime/Online/BuildPatchServices/Private/BuildPatchChunk.cpp
+++ b/Engine/Source/Runtime/Online/BuildPatchServices/Private/BuildPatchChunk.cpp
@@ -19,17 +19,17 @@ FSHAHashData::FSHAHashData()
 	FMemory::Memset(Hash, 0, FSHA1::DigestSize);
 }
 
-FORCEINLINE bool FSHAHashData::operator==(const FSHAHashData& Other) const
+bool FSHAHashData::operator==(const FSHAHashData& Other) const
 {
 	return FMemory::Memcmp(Hash, Other.Hash, FSHA1::DigestSize) == 0;
 }
 
-FORCEINLINE bool FSHAHashData::operator!=(const FSHAHashData& Other) const
+bool FSHAHashData::operator!=(const FSHAHashData& Other) const
 {
 	return !(*this == Other);
 }
 
-FORCEINLINE FString FSHAHashData::ToString() const
+FString FSHAHashData::ToString() const
 {
 	return BytesToHex(Hash, FSHA1::DigestSize);
 }
diff --git a/Engine/Source/Runtime/Online/BuildPatchServices/Private/BuildPatchChunk.h b/Engine/Source/Runtime/Online/BuildPatchServices/Private/BuildPatchChunk.h
--- a/Engine/Source/Runtime/Online/BuildPatchServices/Private/BuildPatchChunk.h
+++ b/Engine/Source/Runtime/Online/BuildPatchServices/Private/BuildPatchChunk.h
@@ -21,9 +21,9 @@ struct FSHAHashData
 
 	FSHAHashData();
 
-	FORCEINLINE bool operator==(const FSHAHashData& Other) const;
-	FORCEINLINE bool operator!=(const FSHAHashData& Other) const;
-	FORCEINLINE FString ToString() const;
+	bool operator==(const FSHAHashData& Other) const;
+	bool operator!=(const FSHAHashData& Other) const;
+	FString ToString() const;
 };
 
 /**
diff --git a/Engine/Source/Runtime/Slate/Private/Framework/Application/SlateApplication.cpp b/Engine/Source/Runtime/Slate/Private/Framework/Application/SlateApplication.cpp
--- a/Engine/Source/Runtime/Slate/Private/Framework/Application/SlateApplication.cpp
+++ b/Engine/Source/Runtime/Slate/Private/Framework/Application/SlateApplication.cpp
@@ -9,11 +9,17 @@
 #include "WebBrowserModule.h"
 #include "IWebBrowserSingleton.h"
 #endif
-
+#include "IWidgetReflector.h"
+#include "GenericCommands.h"
+#include "NotificationManager.h"
 
 DECLARE_CYCLE_STAT( TEXT("Message Tick Time"), STAT_SlateMessageTick, STATGROUP_Slate );
 DECLARE_CYCLE_STAT( TEXT("Update Tooltip Time"), STAT_SlateUpdateTooltip, STATGROUP_Slate );
 DECLARE_CYCLE_STAT( TEXT("Tick Window And Children Time"), STAT_SlateTickWindowAndChildren, STATGROUP_Slate );
+DECLARE_CYCLE_STAT( TEXT("Total Slate Tick Time"), STAT_SlateTickTime, STATGROUP_Slate );
+DECLARE_CYCLE_STAT( TEXT("SlatePrepass"), STAT_SlateCacheDesiredSize, STATGROUP_Slate );
+DECLARE_CYCLE_STAT( TEXT("TickWidgets"), STAT_SlateTickWidgets, STATGROUP_Slate );
+DECLARE_CYCLE_STAT( TEXT("DrawWindows"), STAT_SlateDrawWindowTime, STATGROUP_Slate );
 
 
 // Slate Event Logging is enabled to allow crash log dumping
@@ -4724,3 +4730,13 @@ void FSlateApplication::InitializeAsStandaloneApplication( const TSharedRef<FSla
 	// set the normal UE4 GIsRequestingExit when outer frame is closed
 	FSlateApplication::Get().SetExitRequestedHandler(FSimpleDelegate::CreateStatic(&OnRequestExit));
 }
+
+void FSlateApplication::SetWidgetReflector(const TSharedRef<IWidgetReflector>& WidgetReflector)
+{
+	if (SourceCodeAccessDelegate.IsBound())
+	{
+		WidgetReflector->SetSourceAccessDelegate(SourceCodeAccessDelegate);
+	}
+
+	WidgetReflectorPtr = WidgetReflector;
+}
diff --git a/Engine/Source/Runtime/Slate/Private/Framework/Commands/UICommandDragDropOp.cpp b/Engine/Source/Runtime/Slate/Private/Framework/Commands/UICommandDragDropOp.cpp
--- a/Engine/Source/Runtime/Slate/Private/Framework/Commands/UICommandDragDropOp.cpp
+++ b/Engine/Source/Runtime/Slate/Private/Framework/Commands/UICommandDragDropOp.cpp
@@ -1,6 +1,7 @@
 // Copyright 1998-2014 Epic Games, Inc. All Rights Reserved.
 
 #include "SlatePrivatePCH.h"
+#include "UICommandDragDropOp.h"
 
 
 TSharedRef<FUICommandDragDropOp> FUICommandDragDropOp::New( TSharedRef<const FUICommandInfo> InCommandInfo, FName InOriginMultiBox, TSharedPtr<SWidget> CustomDectorator, FVector2D DecoratorOffset )
diff --git a/Engine/Source/Runtime/Slate/Private/Framework/Docking/FDockingDragOperation.cpp b/Engine/Source/Runtime/Slate/Private/Framework/Docking/FDockingDragOperation.cpp
--- a/Engine/Source/Runtime/Slate/Private/Framework/Docking/FDockingDragOperation.cpp
+++ b/Engine/Source/Runtime/Slate/Private/Framework/Docking/FDockingDragOperation.cpp
@@ -332,7 +332,7 @@ FDockingDragOperation::FDockingDragOperation( const TSharedRef<SDockTab>& InTabT
 /** @return The offset into the tab where the user grabbed in Slate Units. */
 const FVector2D FDockingDragOperation::GetDecoratorOffsetFromCursor()
 {
-	const ETabRole::Type RoleToUse = TabBeingDragged->IsNomadTabWithMajorTabStyle() ? ETabRole::MajorTab : TabBeingDragged->GetTabRole();
+	const ETabRole RoleToUse = TabBeingDragged->IsNomadTabWithMajorTabStyle() ? ETabRole::MajorTab : TabBeingDragged->GetTabRole();
 	const FVector2D TabDesiredSize = TabBeingDragged->GetDesiredSize();
 	const FVector2D MaxTabSize = FDockingConstants::GetMaxTabSizeFor(RoleToUse);
 
diff --git a/Engine/Source/Runtime/Slate/Private/Framework/Docking/SDockTabStack.cpp b/Engine/Source/Runtime/Slate/Private/Framework/Docking/SDockTabStack.cpp
--- a/Engine/Source/Runtime/Slate/Private/Framework/Docking/SDockTabStack.cpp
+++ b/Engine/Source/Runtime/Slate/Private/Framework/Docking/SDockTabStack.cpp
@@ -1,7 +1,7 @@
 // Copyright 1998-2014 Epic Games, Inc. All Rights Reserved.
 
 #include "SlatePrivatePCH.h"
-
+#include "SDockTabStack.h"
 
 void SDockTabStack::Construct( const FArguments& InArgs )
 {
diff --git a/Engine/Source/Runtime/Slate/Private/Framework/Docking/SDockingTabWell.cpp b/Engine/Source/Runtime/Slate/Private/Framework/Docking/SDockingTabWell.cpp
--- a/Engine/Source/Runtime/Slate/Private/Framework/Docking/SDockingTabWell.cpp
+++ b/Engine/Source/Runtime/Slate/Private/Framework/Docking/SDockingTabWell.cpp
@@ -7,7 +7,7 @@
 const FVector2D FDockingConstants::MaxMinorTabSize(150.f, 50.0f);
 const FVector2D FDockingConstants::MaxMajorTabSize(210.f, 50.f);
 
-const FVector2D FDockingConstants::GetMaxTabSizeFor( ETabRole::Type TabRole )
+const FVector2D FDockingConstants::GetMaxTabSizeFor( ETabRole TabRole )
 {
 	return (TabRole == ETabRole::MajorTab)
 		? MaxMajorTabSize
@@ -213,7 +213,7 @@ FVector2D SDockingTabWell::ComputeChildSize( const FGeometry& AllottedGeometry )
 	FVector2D MaxTabSize(0,0);
 	if ( FirstTab.IsValid() )
 	{
-		const ETabRole::Type RoleToUse = FirstTab->IsNomadTabWithMajorTabStyle() ? ETabRole::MajorTab : FirstTab->GetTabRole();
+		const ETabRole RoleToUse = FirstTab->IsNomadTabWithMajorTabStyle() ? ETabRole::MajorTab : FirstTab->GetTabRole();
 		MaxTabSize = FDockingConstants::GetMaxTabSizeFor(RoleToUse);
 	}
 	else
diff --git a/Engine/Source/Runtime/Slate/Private/Framework/Docking/SDockingTabWell.h b/Engine/Source/Runtime/Slate/Private/Framework/Docking/SDockingTabWell.h
--- a/Engine/Source/Runtime/Slate/Private/Framework/Docking/SDockingTabWell.h
+++ b/Engine/Source/Runtime/Slate/Private/Framework/Docking/SDockingTabWell.h
@@ -2,12 +2,13 @@
 
 #pragma once
 
+#include "SDockTab.h"
 
 struct FDockingConstants
 {
 	static const FVector2D MaxMinorTabSize;
 	static const FVector2D MaxMajorTabSize;
-	static const FVector2D SLATE_API GetMaxTabSizeFor( ETabRole::Type TabRole );
+	static const FVector2D SLATE_API GetMaxTabSizeFor( ETabRole TabRole );
 };
 
 
diff --git a/Engine/Source/Runtime/Slate/Private/Framework/Layout/ScrollyZoomy.cpp b/Engine/Source/Runtime/Slate/Private/Framework/Layout/ScrollyZoomy.cpp
--- a/Engine/Source/Runtime/Slate/Private/Framework/Layout/ScrollyZoomy.cpp
+++ b/Engine/Source/Runtime/Slate/Private/Framework/Layout/ScrollyZoomy.cpp
@@ -1,6 +1,7 @@
 // Copyright 1998-2014 Epic Games, Inc. All Rights Reserved.
 
 #include "SlatePrivatePCH.h"
+#include "ScrollyZoomy.h"
 
 
 /* FScrollyZoomy structors
diff --git a/Engine/Source/Runtime/Slate/Private/Framework/MultiBox/MultiBox.cpp b/Engine/Source/Runtime/Slate/Private/Framework/MultiBox/MultiBox.cpp
--- a/Engine/Source/Runtime/Slate/Private/Framework/MultiBox/MultiBox.cpp
+++ b/Engine/Source/Runtime/Slate/Private/Framework/MultiBox/MultiBox.cpp
@@ -8,6 +8,8 @@
 #include "SWidgetBlock.h"
 #include "MultiBoxCustomization.h"
 #include "SClippingHorizontalBox.h"
+#include "UICommandInfo.h"
+#include "UICommandDragDropOp.h"
 
 
 TAttribute<bool> FMultiBoxSettings::UseSmallToolBarIcons;
diff --git a/Engine/Source/Runtime/Slate/Private/Framework/MultiBox/MultiBoxCustomization.cpp b/Engine/Source/Runtime/Slate/Private/Framework/MultiBox/MultiBoxCustomization.cpp
--- a/Engine/Source/Runtime/Slate/Private/Framework/MultiBox/MultiBoxCustomization.cpp
+++ b/Engine/Source/Runtime/Slate/Private/Framework/MultiBox/MultiBoxCustomization.cpp
@@ -5,6 +5,7 @@
 #include "MultiBoxCustomization.h"
 #include "Json.h"
 #include "RemoteConfigIni.h"
+#include "UICommandDragDropOp.h"
 
 void SCustomToolbarPreviewWidget::Construct( const FArguments& InArgs )
 {
diff --git a/Engine/Source/Runtime/Slate/Private/Framework/Notifications/NotificationManager.cpp b/Engine/Source/Runtime/Slate/Private/Framework/Notifications/NotificationManager.cpp
--- a/Engine/Source/Runtime/Slate/Private/Framework/Notifications/NotificationManager.cpp
+++ b/Engine/Source/Runtime/Slate/Private/Framework/Notifications/NotificationManager.cpp
@@ -1,7 +1,8 @@
 // Copyright 1998-2014 Epic Games, Inc. All Rights Reserved.
 
 #include "SlatePrivatePCH.h"
-
+#include "NotificationManager.h"
+#include "SNotificationList.h"
 
 namespace NotificationManagerConstants
 {
diff --git a/Engine/Source/Runtime/Slate/Private/Framework/Text/SlateHyperlinkRun.cpp b/Engine/Source/Runtime/Slate/Private/Framework/Text/SlateHyperlinkRun.cpp
--- a/Engine/Source/Runtime/Slate/Private/Framework/Text/SlateHyperlinkRun.cpp
+++ b/Engine/Source/Runtime/Slate/Private/Framework/Text/SlateHyperlinkRun.cpp
@@ -5,6 +5,7 @@
 #if WITH_FANCY_TEXT
 
 #include "SlateHyperlinkRun.h"
+#include "SRichTextHyperlink.h"
 
 TSharedRef< FSlateHyperlinkRun > FSlateHyperlinkRun::Create( const FRunInfo& InRunInfo, const TSharedRef< const FString >& InText, const FHyperlinkStyle& InStyle, FOnClick NavigateDelegate, FOnGenerateTooltip InTooltipDelegate, FOnGetTooltipText InTooltipTextDelegate )
 {
diff --git a/Engine/Source/Runtime/Slate/Private/SlateModule.cpp b/Engine/Source/Runtime/Slate/Private/SlateModule.cpp
--- a/Engine/Source/Runtime/Slate/Private/SlateModule.cpp
+++ b/Engine/Source/Runtime/Slate/Private/SlateModule.cpp
@@ -4,29 +4,6 @@
 #include "ModuleManager.h"
 
 
-DEFINE_STAT(STAT_SlateCacheDesiredSize_STextBlock);
-DEFINE_STAT(STAT_SlateCacheDesiredSize);
-
-DEFINE_STAT(STAT_SlateOnPaint_SGraphPanel);
-DEFINE_STAT(STAT_SlateOnPaint_STextBlock);
-DEFINE_STAT(STAT_SlateOnPaint_SRichTextBlock);
-DEFINE_STAT(STAT_SlateOnPaint_SEditableText);
-DEFINE_STAT(STAT_SlateOnPaint_SOverlay);
-DEFINE_STAT(STAT_SlateOnPaint_SPanel);
-DEFINE_STAT(STAT_SlateOnPaint_SBorder);
-DEFINE_STAT(STAT_SlateOnPaint_SCompoundWidget);
-DEFINE_STAT(STAT_SlateOnPaint_SBox);
-DEFINE_STAT(STAT_SlateOnPaint_SImage);
-DEFINE_STAT(STAT_SlateOnPaint_SViewport);
-DEFINE_STAT(STAT_SlateOnPaint_SSlider);
-
-DEFINE_STAT(STAT_SlateDrawWindowTime);
-DEFINE_STAT(STAT_SlateTickWidgets);
-DEFINE_STAT(STAT_SlateTickTime);
-
-DEFINE_STAT(STAT_SlateMiscTime); // Useful for locally pref testing miscellaneous Slate areas. Usage of this should not be checked in!
-
-
 /**
  * Implements the Slate module.
  */
diff --git a/Engine/Source/Runtime/Slate/Private/Widgets/Colors/SColorSpectrum.cpp b/Engine/Source/Runtime/Slate/Private/Widgets/Colors/SColorSpectrum.cpp
--- a/Engine/Source/Runtime/Slate/Private/Widgets/Colors/SColorSpectrum.cpp
+++ b/Engine/Source/Runtime/Slate/Private/Widgets/Colors/SColorSpectrum.cpp
@@ -1,6 +1,7 @@
 // Copyright 1998-2014 Epic Games, Inc. All Rights Reserved.
 
 #include "SlatePrivatePCH.h"
+#include "SColorSpectrum.h"
 
 
 /* SColorSpectrum methods
diff --git a/Engine/Source/Runtime/Slate/Private/Widgets/Colors/SColorWheel.cpp b/Engine/Source/Runtime/Slate/Private/Widgets/Colors/SColorWheel.cpp
--- a/Engine/Source/Runtime/Slate/Private/Widgets/Colors/SColorWheel.cpp
+++ b/Engine/Source/Runtime/Slate/Private/Widgets/Colors/SColorWheel.cpp
@@ -1,6 +1,7 @@
 // Copyright 1998-2014 Epic Games, Inc. All Rights Reserved.
 
 #include "SlatePrivatePCH.h"
+#include "SColorWheel.h"
 
 
 /* SColorWheel methods
diff --git a/Engine/Source/Runtime/Slate/Private/Widgets/Docking/SDockTab.cpp b/Engine/Source/Runtime/Slate/Private/Widgets/Docking/SDockTab.cpp
--- a/Engine/Source/Runtime/Slate/Private/Widgets/Docking/SDockTab.cpp
+++ b/Engine/Source/Runtime/Slate/Private/Widgets/Docking/SDockTab.cpp
@@ -1,6 +1,7 @@
 // Copyright 1998-2014 Epic Games, Inc. All Rights Reserved.
 
 #include "SlatePrivatePCH.h"
+#include "SDockTab.h"
 #include "DockingPrivate.h"
 
 
@@ -182,7 +183,7 @@ bool SDockTab::IsForeground() const
 	return ParentPtr.IsValid() ? (ParentPtr.Pin()->GetForegroundTab() == SharedThis(this)) : true;
 }
 
-ETabRole::Type SDockTab::GetTabRole() const
+ETabRole SDockTab::GetTabRole() const
 {
 	return TabRole;
 }
@@ -308,7 +309,7 @@ SDockTab::SDockTab()
 
 }
 
-void SDockTab::ActivateInParent(ETabActivationCause::Type InActivationCause)
+void SDockTab::ActivateInParent(ETabActivationCause InActivationCause)
 {
 	TSharedPtr<SDockingTabWell> ParentTabWell = ParentPtr.Pin();
 	if (ParentTabWell.IsValid())
@@ -709,4 +710,4 @@ FVector2D SDockTab::GetAnimatedScale() const
 	static FVector2D FullyOpen = FVector2D::UnitVector;
 	static FVector2D FullyClosed = FVector2D(1.0f, 0.0f);
 	return FMath::Lerp(FullyClosed, FullyOpen, SpawnAnimCurve.GetLerp());
-}
\ No newline at end of file
+}
diff --git a/Engine/Source/Runtime/Slate/Private/Widgets/Docking/SDockableTab.cpp b/Engine/Source/Runtime/Slate/Private/Widgets/Docking/SDockableTab.cpp
--- a/Engine/Source/Runtime/Slate/Private/Widgets/Docking/SDockableTab.cpp
+++ b/Engine/Source/Runtime/Slate/Private/Widgets/Docking/SDockableTab.cpp
@@ -1,6 +1,8 @@
 // Copyright 1998-2014 Epic Games, Inc. All Rights Reserved.
 
 #include "SlatePrivatePCH.h"
+#include "SDockableTab.h"
+#include "SDockTabStack.h"
 
 
 void SDockableTab::SetContent( TSharedRef<SWidget> InContent )
@@ -38,7 +40,7 @@ void SDockableTab::RemoveTabFromParent()
 
 }
 
-void SDockableTab::ActivateInParent(ETabActivationCause::Type InActivationCause)
+void SDockableTab::ActivateInParent(ETabActivationCause InActivationCause)
 {
 }
 
diff --git a/Engine/Source/Runtime/Slate/Private/Widgets/Images/SImage.cpp b/Engine/Source/Runtime/Slate/Private/Widgets/Images/SImage.cpp
--- a/Engine/Source/Runtime/Slate/Private/Widgets/Images/SImage.cpp
+++ b/Engine/Source/Runtime/Slate/Private/Widgets/Images/SImage.cpp
@@ -2,6 +2,7 @@
 
 #include "SlatePrivatePCH.h"
 
+DECLARE_CYCLE_STAT( TEXT("OnPaint SImage"), STAT_SlateOnPaint_SImage, STATGROUP_Slate );
 
 /**
  * Construct this widget
diff --git a/Engine/Source/Runtime/Slate/Private/Widgets/Input/SEditableText.cpp b/Engine/Source/Runtime/Slate/Private/Widgets/Input/SEditableText.cpp
--- a/Engine/Source/Runtime/Slate/Private/Widgets/Input/SEditableText.cpp
+++ b/Engine/Source/Runtime/Slate/Private/Widgets/Input/SEditableText.cpp
@@ -2,6 +2,9 @@
 
 #include "SlatePrivatePCH.h"
 #include "TextEditHelper.h"
+#include "GenericCommands.h"
+
+DECLARE_CYCLE_STAT( TEXT("OnPaint SEditableText"), STAT_SlateOnPaint_SEditableText, STATGROUP_Slate );
 
 /** A pointer to the editable text currently under the mouse cursor. nullptr when there isn't one. */
 SEditableText* SEditableText::EditableTextUnderCursor = nullptr;
diff --git a/Engine/Source/Runtime/Slate/Private/Widgets/Input/SExpandableButton.cpp b/Engine/Source/Runtime/Slate/Private/Widgets/Input/SExpandableButton.cpp
--- a/Engine/Source/Runtime/Slate/Private/Widgets/Input/SExpandableButton.cpp
+++ b/Engine/Source/Runtime/Slate/Private/Widgets/Input/SExpandableButton.cpp
@@ -1,6 +1,7 @@
 // Copyright 1998-2014 Epic Games, Inc. All Rights Reserved.
 
 #include "SlatePrivatePCH.h"
+#include "SExpandableButton.h"
 
 
 //////////////////////////////////////////////////////////////////////////
diff --git a/Engine/Source/Runtime/Slate/Private/Widgets/Input/SRotatorInputBox.cpp b/Engine/Source/Runtime/Slate/Private/Widgets/Input/SRotatorInputBox.cpp
--- a/Engine/Source/Runtime/Slate/Private/Widgets/Input/SRotatorInputBox.cpp
+++ b/Engine/Source/Runtime/Slate/Private/Widgets/Input/SRotatorInputBox.cpp
@@ -1,6 +1,8 @@
 // Copyright 1998-2014 Epic Games, Inc. All Rights Reserved.
 
 #include "SlatePrivatePCH.h"
+#include "SRotatorInputBox.h"
+#include "SNumericEntryBox.h"
 
 
 #define LOCTEXT_NAMESPACE "SRotatorInputBox"
diff --git a/Engine/Source/Runtime/Slate/Private/Widgets/Input/SSearchBox.cpp b/Engine/Source/Runtime/Slate/Private/Widgets/Input/SSearchBox.cpp
--- a/Engine/Source/Runtime/Slate/Private/Widgets/Input/SSearchBox.cpp
+++ b/Engine/Source/Runtime/Slate/Private/Widgets/Input/SSearchBox.cpp
@@ -1,6 +1,7 @@
 // Copyright 1998-2014 Epic Games, Inc. All Rights Reserved.
 
 #include "SlatePrivatePCH.h"
+#include "SSearchBox.h"
 
 
 void SSearchBox::Construct( const FArguments& InArgs )
@@ -115,8 +116,8 @@ void SSearchBox::Tick( const FGeometry& AllottedGeometry, const double InCurrent
 	CurrentTime = InCurrentTime;
 
 	if ( bTypingFilterText && InCurrentTime > LastTypeTime + FilterDelayAfterTyping)
-	{
-		OnTextChangedDelegate.ExecuteIfBound( LastPendingTextChangedValue );
+{
+	OnTextChangedDelegate.ExecuteIfBound( LastPendingTextChangedValue );
 		bTypingFilterText = false;
 	}
 }
diff --git a/Engine/Source/Runtime/Slate/Private/Widgets/Input/SSlider.cpp b/Engine/Source/Runtime/Slate/Private/Widgets/Input/SSlider.cpp
--- a/Engine/Source/Runtime/Slate/Private/Widgets/Input/SSlider.cpp
+++ b/Engine/Source/Runtime/Slate/Private/Widgets/Input/SSlider.cpp
@@ -2,6 +2,7 @@
 
 #include "SlatePrivatePCH.h"
 
+DECLARE_CYCLE_STAT( TEXT("OnPaint SSlider"), STAT_SlateOnPaint_SSlider, STATGROUP_Slate );
 
 void SSlider::Construct( const SSlider::FArguments& InDeclaration )
 {
diff --git a/Engine/Source/Runtime/Slate/Private/Widgets/Input/SSuggestionTextBox.cpp b/Engine/Source/Runtime/Slate/Private/Widgets/Input/SSuggestionTextBox.cpp
--- a/Engine/Source/Runtime/Slate/Private/Widgets/Input/SSuggestionTextBox.cpp
+++ b/Engine/Source/Runtime/Slate/Private/Widgets/Input/SSuggestionTextBox.cpp
@@ -1,6 +1,7 @@
 // Copyright 1998-2014 Epic Games, Inc. All Rights Reserved.
 
 #include "SlatePrivatePCH.h"
+#include "SSuggestionTextBox.h"
 
 
 /* SSuggestionTextBox structors
diff --git a/Engine/Source/Runtime/Slate/Private/Widgets/Input/STextComboPopup.cpp b/Engine/Source/Runtime/Slate/Private/Widgets/Input/STextComboPopup.cpp
--- a/Engine/Source/Runtime/Slate/Private/Widgets/Input/STextComboPopup.cpp
+++ b/Engine/Source/Runtime/Slate/Private/Widgets/Input/STextComboPopup.cpp
@@ -1,6 +1,7 @@
 // Copyright 1998-2014 Epic Games, Inc. All Rights Reserved.
 
 #include "SlatePrivatePCH.h"
+#include "STextComboPopup.h"
 
 
 void STextComboPopup::Construct( const FArguments& InArgs )
diff --git a/Engine/Source/Runtime/Slate/Private/Widgets/Input/SVectorInputBox.cpp b/Engine/Source/Runtime/Slate/Private/Widgets/Input/SVectorInputBox.cpp
--- a/Engine/Source/Runtime/Slate/Private/Widgets/Input/SVectorInputBox.cpp
+++ b/Engine/Source/Runtime/Slate/Private/Widgets/Input/SVectorInputBox.cpp
@@ -1,6 +1,8 @@
 // Copyright 1998-2014 Epic Games, Inc. All Rights Reserved.
 
 #include "SlatePrivatePCH.h"
+#include "SVectorInputBox.h"
+#include "SNumericEntryBox.h"
 
 
 #define LOCTEXT_NAMESPACE "SVectorInputBox"
diff --git a/Engine/Source/Runtime/Slate/Private/Widgets/Input/SVirtualJoystick.cpp b/Engine/Source/Runtime/Slate/Private/Widgets/Input/SVirtualJoystick.cpp
--- a/Engine/Source/Runtime/Slate/Private/Widgets/Input/SVirtualJoystick.cpp
+++ b/Engine/Source/Runtime/Slate/Private/Widgets/Input/SVirtualJoystick.cpp
@@ -1,6 +1,7 @@
 // Copyright 1998-2014 Epic Games, Inc. All Rights Reserved.
 
 #include "SlatePrivatePCH.h"
+#include "SVirtualJoystick.h"
 
 
 const float OPACITY_LERP_RATE = 3.f;
diff --git a/Engine/Source/Runtime/Slate/Private/Widgets/Input/SVirtualKeyboardEntry.cpp b/Engine/Source/Runtime/Slate/Private/Widgets/Input/SVirtualKeyboardEntry.cpp
--- a/Engine/Source/Runtime/Slate/Private/Widgets/Input/SVirtualKeyboardEntry.cpp
+++ b/Engine/Source/Runtime/Slate/Private/Widgets/Input/SVirtualKeyboardEntry.cpp
@@ -1,6 +1,7 @@
 // Copyright 1998-2014 Epic Games, Inc. All Rights Reserved.
 
 #include "SlatePrivatePCH.h"
+#include "SVirtualKeyboardEntry.h"
 #include "TextEditHelper.h"
 
 /** Constructor */
diff --git a/Engine/Source/Runtime/Slate/Private/Widgets/Input/SVolumeControl.cpp b/Engine/Source/Runtime/Slate/Private/Widgets/Input/SVolumeControl.cpp
--- a/Engine/Source/Runtime/Slate/Private/Widgets/Input/SVolumeControl.cpp
+++ b/Engine/Source/Runtime/Slate/Private/Widgets/Input/SVolumeControl.cpp
@@ -1,6 +1,7 @@
 // Copyright 1998-2014 Epic Games, Inc. All Rights Reserved.
 
 #include "SlatePrivatePCH.h"
+#include "SVolumeControl.h"
 
 
 /**
diff --git a/Engine/Source/Runtime/Slate/Private/Widgets/Layout/SBorder.cpp b/Engine/Source/Runtime/Slate/Private/Widgets/Layout/SBorder.cpp
--- a/Engine/Source/Runtime/Slate/Private/Widgets/Layout/SBorder.cpp
+++ b/Engine/Source/Runtime/Slate/Private/Widgets/Layout/SBorder.cpp
@@ -2,6 +2,7 @@
 
 #include "SlatePrivatePCH.h"
  
+DECLARE_CYCLE_STAT( TEXT("OnPaint SBorder"), STAT_SlateOnPaint_SBorder, STATGROUP_Slate );
 
 SBorder::SBorder()
 	: BorderImage( FCoreStyle::Get().GetBrush( "Border" ) )
diff --git a/Engine/Source/Runtime/Slate/Private/Widgets/Layout/SBox.cpp b/Engine/Source/Runtime/Slate/Private/Widgets/Layout/SBox.cpp
--- a/Engine/Source/Runtime/Slate/Private/Widgets/Layout/SBox.cpp
+++ b/Engine/Source/Runtime/Slate/Private/Widgets/Layout/SBox.cpp
@@ -3,6 +3,8 @@
 #include "SlatePrivatePCH.h"
 #include "LayoutUtils.h"
 
+DECLARE_CYCLE_STAT( TEXT("OnPaint SBox"), STAT_SlateOnPaint_SBox, STATGROUP_Slate );
+
 SBox::SBox()
 : ChildSlot()
 {
diff --git a/Engine/Source/Runtime/Slate/Private/Widgets/Layout/SDPIScaler.cpp b/Engine/Source/Runtime/Slate/Private/Widgets/Layout/SDPIScaler.cpp
--- a/Engine/Source/Runtime/Slate/Private/Widgets/Layout/SDPIScaler.cpp
+++ b/Engine/Source/Runtime/Slate/Private/Widgets/Layout/SDPIScaler.cpp
@@ -1,6 +1,7 @@
 // Copyright 1998-2014 Epic Games, Inc. All Rights Reserved.
 
 #include "SlatePrivatePCH.h"
+#include "SDPIScaler.h"
 
 SDPIScaler::SDPIScaler()
 : ChildSlot()
diff --git a/Engine/Source/Runtime/Slate/Private/Widgets/Layout/SMissingWidget.cpp b/Engine/Source/Runtime/Slate/Private/Widgets/Layout/SMissingWidget.cpp
--- a/Engine/Source/Runtime/Slate/Private/Widgets/Layout/SMissingWidget.cpp
+++ b/Engine/Source/Runtime/Slate/Private/Widgets/Layout/SMissingWidget.cpp
@@ -1,6 +1,7 @@
 // Copyright 1998-2014 Epic Games, Inc. All Rights Reserved.
 
 #include "SlatePrivatePCH.h"
+#include "SMissingWidget.h"
 
 
 #define LOCTEXT_NAMESPACE "SMissingWidget"
diff --git a/Engine/Source/Runtime/Slate/Private/Widgets/Layout/SSafeZone.cpp b/Engine/Source/Runtime/Slate/Private/Widgets/Layout/SSafeZone.cpp
new file mode 100644
--- /dev/null
+++ b/Engine/Source/Runtime/Slate/Private/Widgets/Layout/SSafeZone.cpp
@@ -0,0 +1,31 @@
+// Copyright 1998-2014 Epic Games, Inc. All Rights Reserved.
+
+#include "SlatePrivatePCH.h"
+#include "SSafeZone.h"
+
+
+void SSafeZone::Construct(const FArguments& InArgs)
+{
+	SBox::Construct(SBox::FArguments()
+		.HAlign(InArgs._HAlign)
+		.VAlign(InArgs._VAlign)
+		.Padding(this, &SSafeZone::GetSafeZonePadding)
+		[
+			InArgs._Content.Widget
+		]
+	);
+	
+	IsTitleSafe = InArgs._IsTitleSafe;
+	Padding = InArgs._Padding;
+}
+
+FMargin SSafeZone::GetSafeZonePadding() const
+{
+	// @todo: should we have a function that just returns the safe area size?
+	FDisplayMetrics Metrics;
+	FSlateApplication::Get().GetDisplayMetrics(Metrics);
+
+	// return either the TitleSafe or the ActionSafe size, added to the user padding
+	return Padding.Get() + (IsTitleSafe ? FMargin(Metrics.TitleSafePaddingSize.X, Metrics.TitleSafePaddingSize.Y) :
+		FMargin(Metrics.ActionSafePaddingSize.X, Metrics.ActionSafePaddingSize.Y));
+}
diff --git a/Engine/Source/Runtime/Slate/Private/Widgets/Notifications/SNotificationList.cpp b/Engine/Source/Runtime/Slate/Private/Widgets/Notifications/SNotificationList.cpp
--- a/Engine/Source/Runtime/Slate/Private/Widgets/Notifications/SNotificationList.cpp
+++ b/Engine/Source/Runtime/Slate/Private/Widgets/Notifications/SNotificationList.cpp
@@ -1,6 +1,9 @@
 // Copyright 1998-2014 Epic Games, Inc. All Rights Reserved.
 
 #include "SlatePrivatePCH.h"
+#include "SNotificationList.h"
+#include "SHyperlink.h"
+#include "SThrobber.h"
 
 
 /////////////////////////////////////////////////
diff --git a/Engine/Source/Runtime/Slate/Private/Widgets/SViewport.cpp b/Engine/Source/Runtime/Slate/Private/Widgets/SViewport.cpp
--- a/Engine/Source/Runtime/Slate/Private/Widgets/SViewport.cpp
+++ b/Engine/Source/Runtime/Slate/Private/Widgets/SViewport.cpp
@@ -2,6 +2,7 @@
 
 #include "SlatePrivatePCH.h"
 
+DECLARE_CYCLE_STAT( TEXT("OnPaint SViewport"), STAT_SlateOnPaint_SViewport, STATGROUP_Slate );
 
 /* SViewport structors
  *****************************************************************************/
diff --git a/Engine/Source/Runtime/Slate/Private/Widgets/Text/SRichTextBlock.cpp b/Engine/Source/Runtime/Slate/Private/Widgets/Text/SRichTextBlock.cpp
--- a/Engine/Source/Runtime/Slate/Private/Widgets/Text/SRichTextBlock.cpp
+++ b/Engine/Source/Runtime/Slate/Private/Widgets/Text/SRichTextBlock.cpp
@@ -4,6 +4,8 @@
 
 #if WITH_FANCY_TEXT
 
+DECLARE_CYCLE_STAT( TEXT("OnPaint SRichTextBlock"), STAT_SlateOnPaint_SRichTextBlock, STATGROUP_Slate );
+
 #include "TextLayoutEngine.h"
 #include "TextBlockLayout.h"
 #include "IRichTextMarkupParser.h"
diff --git a/Engine/Source/Runtime/Slate/Private/Widgets/Text/STextBlock.cpp b/Engine/Source/Runtime/Slate/Private/Widgets/Text/STextBlock.cpp
--- a/Engine/Source/Runtime/Slate/Private/Widgets/Text/STextBlock.cpp
+++ b/Engine/Source/Runtime/Slate/Private/Widgets/Text/STextBlock.cpp
@@ -4,6 +4,9 @@
 #include "PlainTextLayoutMarshaller.h"
 #include "TextBlockLayout.h"
 
+DECLARE_CYCLE_STAT( TEXT("STextBlock"), STAT_SlateComputeDesiredSize_STextBlock, STATGROUP_Slate );
+DECLARE_CYCLE_STAT( TEXT("OnPaint STextBlock"), STAT_SlateOnPaint_STextBlock, STATGROUP_Slate );
+
 void STextBlock::Construct( const FArguments& InArgs )
 {
 	TextStyle = InArgs._TextStyle;
diff --git a/Engine/Source/Runtime/Slate/Public/Framework/Application/IWidgetReflector.h b/Engine/Source/Runtime/Slate/Public/Framework/Application/IWidgetReflector.h
--- a/Engine/Source/Runtime/Slate/Public/Framework/Application/IWidgetReflector.h
+++ b/Engine/Source/Runtime/Slate/Public/Framework/Application/IWidgetReflector.h
@@ -3,9 +3,6 @@
 #pragma once
 
 
-/** A Delegate for passing along a string of a source code location to access */
-DECLARE_DELEGATE_RetVal_ThreeParams(bool, FAccessSourceCode, const FString& /*FileName*/, int32 /*InLineNumber*/, int32 /*InColumnNumber*/);
-
 
 /**
  * Interface for widget reflectors.
diff --git a/Engine/Source/Runtime/Slate/Public/Framework/Application/SlateApplication.h b/Engine/Source/Runtime/Slate/Public/Framework/Application/SlateApplication.h
--- a/Engine/Source/Runtime/Slate/Public/Framework/Application/SlateApplication.h
+++ b/Engine/Source/Runtime/Slate/Public/Framework/Application/SlateApplication.h
@@ -3,13 +3,15 @@
 #pragma once
 
 #include "GenericApplicationMessageHandler.h"
-
+#include "MenuStack.h"
 
 class SToolTip;
 class SViewport;
 class SWindow;
 class FHittestGrid;
-
+struct FPopupTransitionEffect;
+class FMenuStack;
+class IWidgetReflector;
 
 /** A Delegate for querying whether source code access is possible */
 DECLARE_DELEGATE_RetVal(bool, FQueryAccessSourceCode);
@@ -467,15 +469,7 @@ public:
 	 *
 	 * @param WidgetReflector The widget reflector to set.
 	 */
-	void SetWidgetReflector( const TSharedRef<IWidgetReflector>& WidgetReflector )
-	{
-		if (SourceCodeAccessDelegate.IsBound())
-		{
-			WidgetReflector->SetSourceAccessDelegate(SourceCodeAccessDelegate);
-		}
-
-		WidgetReflectorPtr = WidgetReflector;
-	}
+	void SetWidgetReflector( const TSharedRef<IWidgetReflector>& WidgetReflector );
 
 	/** @param AccessDelegate The delegate to pass along to the widget reflector */
 	void SetWidgetReflectorSourceAccessDelegate(FAccessSourceCode AccessDelegate)
diff --git a/Engine/Source/Runtime/Slate/Public/Framework/Commands/InputGesture.h b/Engine/Source/Runtime/Slate/Public/Framework/Commands/InputGesture.h
--- a/Engine/Source/Runtime/Slate/Public/Framework/Commands/InputGesture.h
+++ b/Engine/Source/Runtime/Slate/Public/Framework/Commands/InputGesture.h
@@ -1,6 +1,7 @@
 // Copyright 1998-2014 Epic Games, Inc. All Rights Reserved.
 
 #pragma once
+#include "GenericApplication.h"
 
 
 /**
